Fix compiler warning

diff --git a/Engine/Plugins/Runtime/OpenXR/Source/OpenXRInput/Private/OpenXRInput.cpp b/Engine/Plugins/Runtime/OpenXR/Source/OpenXRInput/Private/OpenXRInput.cpp
--- a/Engine/Plugins/Runtime/OpenXR/Source/OpenXRInput/Private/OpenXRInput.cpp
+++ b/Engine/Plugins/Runtime/OpenXR/Source/OpenXRInput/Private/OpenXRInput.cpp
@@ -34,7 +34,7 @@ FORCEINLINE void FilterActionName(const char* InActionName, char* OutActionName)
 	size_t i;
 	for (i = 0; InActionName[i] != '\0' && i < XR_MAX_ACTION_NAME_SIZE - 1; i++)
 	{
-		const char c = InActionName[i];
+		unsigned char c = InActionName[i];
 		OutActionName[i] = (c == ' ') ? '-' : isalnum(c) ? tolower(c) : '_';
 	}
 	OutActionName[i] = '\0';
