Fix clang compiler error

diff --git a/Engine/Plugins/Runtime/OpenXR/Source/OpenXRHMD/Private/OpenXRHMD.cpp b/Engine/Plugins/Runtime/OpenXR/Source/OpenXRHMD/Private/OpenXRHMD.cpp
--- a/Engine/Plugins/Runtime/OpenXR/Source/OpenXRHMD/Private/OpenXRHMD.cpp
+++ b/Engine/Plugins/Runtime/OpenXR/Source/OpenXRHMD/Private/OpenXRHMD.cpp
@@ -811,9 +811,11 @@ void FOpenXRHMD::OnBeginRendering_RenderThread(FRHICommandListImmediate& RHICmdL
 		Projection.subImage.swapchain = Swapchain->Handle;
 		Projection.subImage.imageArrayIndex = 0;
 		Projection.subImage.imageRect = {
-			OffsetX, 0,
-			(int32)Config.recommendedImageRectWidth,
-			(int32)Config.recommendedImageRectHeight
+			{ OffsetX, 0 },
+			{
+				(int32)Config.recommendedImageRectWidth,
+				(int32)Config.recommendedImageRectHeight
+			}
 		};
 		OffsetX += Config.recommendedImageRectWidth;
 	}
