Move FindSceneViewport to an anonymous namespace

diff --git a/Engine/Plugins/Runtime/OpenXR/Source/OpenXRHMD/Private/OpenXRHMD.cpp b/Engine/Plugins/Runtime/OpenXR/Source/OpenXRHMD/Private/OpenXRHMD.cpp
--- a/Engine/Plugins/Runtime/OpenXR/Source/OpenXRHMD/Private/OpenXRHMD.cpp
+++ b/Engine/Plugins/Runtime/OpenXR/Source/OpenXRHMD/Private/OpenXRHMD.cpp
@@ -36,39 +36,39 @@ namespace {
 		{ DXGI_FORMAT_R8G8B8A8_UNORM_SRGB, PF_R8G8B8A8 },
 		{ DXGI_FORMAT_B8G8R8A8_UNORM_SRGB, PF_B8G8R8A8 },
     };
-}
 
-/** Helper function for acquiring the appropriate FSceneViewport */
-FSceneViewport* FindSceneViewport()
-{
-	if (!GIsEditor)
+	/** Helper function for acquiring the appropriate FSceneViewport */
+	FSceneViewport* FindSceneViewport()
 	{
-		UGameEngine* GameEngine = Cast<UGameEngine>(GEngine);
-		return GameEngine->SceneViewport.Get();
-	}
-#if WITH_EDITOR
-	else
-	{
-		UEditorEngine* EditorEngine = CastChecked<UEditorEngine>(GEngine);
-		FSceneViewport* PIEViewport = (FSceneViewport*)EditorEngine->GetPIEViewport();
-		if (PIEViewport != nullptr && PIEViewport->IsStereoRenderingAllowed())
+		if (!GIsEditor)
 		{
-			// PIE is setup for stereo rendering
-			return PIEViewport;
+			UGameEngine* GameEngine = Cast<UGameEngine>(GEngine);
+			return GameEngine->SceneViewport.Get();
 		}
+	#if WITH_EDITOR
 		else
 		{
-			// Check to see if the active editor viewport is drawing in stereo mode
-			// @todo vreditor: Should work with even non-active viewport!
-			FSceneViewport* EditorViewport = (FSceneViewport*)EditorEngine->GetActiveViewport();
-			if (EditorViewport != nullptr && EditorViewport->IsStereoRenderingAllowed())
+			UEditorEngine* EditorEngine = CastChecked<UEditorEngine>(GEngine);
+			FSceneViewport* PIEViewport = (FSceneViewport*)EditorEngine->GetPIEViewport();
+			if (PIEViewport != nullptr && PIEViewport->IsStereoRenderingAllowed())
 			{
-				return EditorViewport;
+				// PIE is setup for stereo rendering
+				return PIEViewport;
+			}
+			else
+			{
+				// Check to see if the active editor viewport is drawing in stereo mode
+				// @todo vreditor: Should work with even non-active viewport!
+				FSceneViewport* EditorViewport = (FSceneViewport*)EditorEngine->GetActiveViewport();
+				if (EditorViewport != nullptr && EditorViewport->IsStereoRenderingAllowed())
+				{
+					return EditorViewport;
+				}
 			}
 		}
+	#endif
+		return nullptr;
 	}
-#endif
-	return nullptr;
 }
 
 //---------------------------------------------------
