Add error messaging in case PreInit() fails.

diff --git a/Engine/Plugins/Runtime/OpenXR/Source/OpenXRHMD/Private/OpenXRHMD.cpp b/Engine/Plugins/Runtime/OpenXR/Source/OpenXRHMD/Private/OpenXRHMD.cpp
--- a/Engine/Plugins/Runtime/OpenXR/Source/OpenXRHMD/Private/OpenXRHMD.cpp
+++ b/Engine/Plugins/Runtime/OpenXR/Source/OpenXRHMD/Private/OpenXRHMD.cpp
@@ -172,8 +172,12 @@ bool FOpenXRHMDPlugin::PreInit()
 	Info.enabledApiLayerNames = nullptr;
 	Info.enabledExtensionCount = 1;
 	Info.enabledExtensionNames = extensions;
-	if (!XR_ENSURE(xrCreateInstance(&Info, &Instance)))
+	XrResult rs = xrCreateInstance(&Info, &Instance);
+	if (XR_FAILED(rs))
 	{
+		char error[XR_MAX_RESULT_STRING_SIZE] = { '\0' };
+		xrResultToString(XR_NULL_HANDLE, rs, error);
+		UE_LOG(LogHMD, Log, TEXT("Failed to create an OpenXR instance, result is %s. Please check if you have an OpenXR runtime installed."), error);
 		return false;
 	}
 
@@ -181,8 +185,12 @@ bool FOpenXRHMDPlugin::PreInit()
 	SystemInfo.type = XR_TYPE_SYSTEM_GET_INFO;
 	SystemInfo.next = nullptr;
 	SystemInfo.formFactor = XR_FORM_FACTOR_HEAD_MOUNTED_DISPLAY;
-	if (!XR_ENSURE(xrGetSystem(Instance, &SystemInfo, &System)))
+	rs = xrGetSystem(Instance, &SystemInfo, &System);
+	if (XR_FAILED(rs))
 	{
+		char error[XR_MAX_RESULT_STRING_SIZE] = { '\0' };
+		xrResultToString(XR_NULL_HANDLE, rs, error);
+		UE_LOG(LogHMD, Log, TEXT("Failed to get an OpenXR system, result is %s. Please check that your runtime supports VR headsets."), error);
 		return false;
 	}
 
