Merging OpenXR plugin updates from Dev-VR:

CL 6092122 - Correct the stereo projection matrix field-of-view
CL 6092136 - Use the parent for scalars instead of providing suggested bindings for both child paths
CL 6103424 - Refactor OpenXR version information to match VulkanRHI more closely

diff --git a/Engine/Plugins/Runtime/OpenXR/Source/OpenXRHMD/Private/OpenXRHMD.cpp b/Engine/Plugins/Runtime/OpenXR/Source/OpenXRHMD/Private/OpenXRHMD.cpp
--- a/Engine/Plugins/Runtime/OpenXR/Source/OpenXRHMD/Private/OpenXRHMD.cpp
+++ b/Engine/Plugins/Runtime/OpenXR/Source/OpenXRHMD/Private/OpenXRHMD.cpp
@@ -156,17 +156,25 @@ bool FOpenXRHMDPlugin::PreInit()
 		return false;
 	}
 
-	FString AppName = FApp::GetName();
 	const char* extensions[] = { XR_KHR_D3D11_ENABLE_EXTENSION_NAME };
 
+	// Engine registration can be disabled via console var.
+	auto* CVarDisableEngineAndAppRegistration = IConsoleManager::Get().FindTConsoleVariableDataInt(TEXT("r.DisableEngineAndAppRegistration"));
+	bool bDisableEngineRegistration = (CVarDisableEngineAndAppRegistration && CVarDisableEngineAndAppRegistration->GetValueOnAnyThread() != 0);
+
+	// EngineName will be of the form "UnrealEngine4.21", with the minor version ("21" in this example)
+	// updated with every quarterly release
+	FString EngineName = bDisableEngineRegistration ? FString("") : FApp::GetEpicProductIdentifier() + FEngineVersion::Current().ToString(EVersionComponent::Minor);
+	const TCHAR* ProjectName = bDisableEngineRegistration ? TEXT("") : FApp::GetProjectName();
+
 	XrInstanceCreateInfo Info;
 	Info.type = XR_TYPE_INSTANCE_CREATE_INFO;
 	Info.next = nullptr;
 	Info.createFlags = 0;
-	FPlatformString::Convert(Info.applicationInfo.applicationName, XR_MAX_APPLICATION_NAME_SIZE, GetData(AppName), AppName.Len() + 1);
-	Info.applicationInfo.applicationVersion = 0;
-	FCStringAnsi::Strcpy(Info.applicationInfo.engineName, XR_MAX_ENGINE_NAME_SIZE, "Unreal Engine");
-	Info.applicationInfo.engineVersion = (uint32)FEngineVersion::Current().GetMajor() << 16 | FEngineVersion::Current().GetMinor();
+	FTCHARToUTF8_Convert::Convert(Info.applicationInfo.applicationName, XR_MAX_APPLICATION_NAME_SIZE, ProjectName, TCString<TCHAR>::Strlen(ProjectName) + 1);
+	Info.applicationInfo.applicationVersion = FEngineVersion::Current().GetChangelist();
+	FTCHARToUTF8_Convert::Convert(Info.applicationInfo.engineName, XR_MAX_ENGINE_NAME_SIZE, *EngineName, EngineName.Len() + 1);
+	Info.applicationInfo.engineVersion = (uint32)(FEngineVersion::Current().GetMinor() << 16 | FEngineVersion::Current().GetPatch());
 	Info.applicationInfo.apiVersion = XR_CURRENT_API_VERSION;
 	Info.enabledApiLayerCount = 0;
 	Info.enabledApiLayerNames = nullptr;
@@ -479,18 +487,18 @@ FMatrix FOpenXRHMD::GetStereoProjectionMatrix(const enum EStereoscopicPass Stere
 
 	Fov.angleUp = tan(Fov.angleUp);
 	Fov.angleDown = tan(Fov.angleDown);
-	Fov.angleLeft = tan(-Fov.angleLeft);
-	Fov.angleRight = tan(-Fov.angleRight);
+	Fov.angleLeft = tan(Fov.angleLeft);
+	Fov.angleRight = tan(Fov.angleRight);
 
-	float SumRL = (Fov.angleLeft + Fov.angleRight);
+	float SumRL = (Fov.angleRight + Fov.angleLeft);
 	float SumTB = (Fov.angleUp + Fov.angleDown);
-	float InvRL = (1.0f / (Fov.angleLeft - Fov.angleRight));
+	float InvRL = (1.0f / (Fov.angleRight - Fov.angleLeft));
 	float InvTB = (1.0f / (Fov.angleUp - Fov.angleDown));
 
 	FMatrix Mat = FMatrix(
 		FPlane((2.0f * InvRL), 0.0f, 0.0f, 0.0f),
 		FPlane(0.0f, (2.0f * InvTB), 0.0f, 0.0f),
-		FPlane((SumRL * InvRL), (SumTB * InvTB), 0.0f, 1.0f),
+		FPlane((SumRL * -InvRL), (SumTB * -InvTB), 0.0f, 1.0f),
 		FPlane(0.0f, 0.0f, ZNear, 0.0f)
 	);
 
diff --git a/Engine/Plugins/Runtime/OpenXR/Source/OpenXRInput/Private/OpenXRInput.cpp b/Engine/Plugins/Runtime/OpenXR/Source/OpenXRInput/Private/OpenXRInput.cpp
--- a/Engine/Plugins/Runtime/OpenXR/Source/OpenXRInput/Private/OpenXRInput.cpp
+++ b/Engine/Plugins/Runtime/OpenXR/Source/OpenXRInput/Private/OpenXRInput.cpp
@@ -176,12 +176,10 @@ FOpenXRInputPlugin::FOpenXRInput::FOpenXRInput(FOpenXRHMD* HMD)
 	Bindings.Add(XrActionSuggestedBinding{ Controllers[EControllerHand::Right].Vibration, GetPath(Instance, "/user/hand/right/output/haptic") });
 
 	InteractionMappings.Add(FGamepadKeyNames::MotionController_Left_Shoulder, GetPath(Instance, "/user/hand/left/input/menu/click"));
-	InteractionMappings.Add(FGamepadKeyNames::MotionController_Left_Trigger, GetPath(Instance, "/user/hand/left/input/trigger/value"));
-	InteractionMappings.Add(FGamepadKeyNames::MotionController_Left_TriggerAxis, GetPath(Instance, "/user/hand/left/input/trigger/value"));
-	InteractionMappings.Add(FGamepadKeyNames::MotionController_Left_Grip1, GetPath(Instance, "/user/hand/left/input/grip/value"));
-	InteractionMappings.Add(FGamepadKeyNames::MotionController_Left_Grip1Axis, GetPath(Instance, "/user/hand/left/input/grip/value"));
-	InteractionMappings.Add(FGamepadKeyNames::MotionController_Left_Grip1, GetPath(Instance, "/user/hand/left/input/grip/click"));
-	InteractionMappings.Add(FGamepadKeyNames::MotionController_Left_Grip1Axis, GetPath(Instance, "/user/hand/left/input/grip/click"));
+	InteractionMappings.Add(FGamepadKeyNames::MotionController_Left_Trigger, GetPath(Instance, "/user/hand/left/input/trigger"));
+	InteractionMappings.Add(FGamepadKeyNames::MotionController_Left_TriggerAxis, GetPath(Instance, "/user/hand/left/input/trigger"));
+	InteractionMappings.Add(FGamepadKeyNames::MotionController_Left_Grip1, GetPath(Instance, "/user/hand/left/input/grip"));
+	InteractionMappings.Add(FGamepadKeyNames::MotionController_Left_Grip1Axis, GetPath(Instance, "/user/hand/left/input/grip"));
 	InteractionMappings.Add(FGamepadKeyNames::MotionController_Left_Thumbstick_X, GetPath(Instance, "/user/hand/left/input/thumbstick/x"));
 	InteractionMappings.Add(FGamepadKeyNames::MotionController_Left_Thumbstick_Y, GetPath(Instance, "/user/hand/left/input/thumbstick/y"));
 	InteractionMappings.Add(FGamepadKeyNames::MotionController_Left_Thumbstick, GetPath(Instance, "/user/hand/left/input/thumbstick/click"));
@@ -189,12 +187,10 @@ FOpenXRInputPlugin::FOpenXRInput::FOpenXRInput(FOpenXRHMD* HMD)
 	InteractionMappings.Add(FGamepadKeyNames::MotionController_Left_FaceButton1, GetPath(Instance, "/user/hand/left/input/y/click"));
 
 	InteractionMappings.Add(FGamepadKeyNames::MotionController_Right_Shoulder, GetPath(Instance, "/user/hand/right/input/menu/click"));
-	InteractionMappings.Add(FGamepadKeyNames::MotionController_Right_Trigger, GetPath(Instance, "/user/hand/right/input/trigger/value"));
-	InteractionMappings.Add(FGamepadKeyNames::MotionController_Right_TriggerAxis, GetPath(Instance, "/user/hand/right/input/trigger/value"));
-	InteractionMappings.Add(FGamepadKeyNames::MotionController_Right_Grip1, GetPath(Instance, "/user/hand/right/input/grip/value"));
-	InteractionMappings.Add(FGamepadKeyNames::MotionController_Right_Grip1Axis, GetPath(Instance, "/user/hand/right/input/grip/value"));
-	InteractionMappings.Add(FGamepadKeyNames::MotionController_Right_Grip1, GetPath(Instance, "/user/hand/right/input/grip/click"));
-	InteractionMappings.Add(FGamepadKeyNames::MotionController_Right_Grip1Axis, GetPath(Instance, "/user/hand/right/input/grip/click"));
+	InteractionMappings.Add(FGamepadKeyNames::MotionController_Right_Trigger, GetPath(Instance, "/user/hand/right/input/trigger"));
+	InteractionMappings.Add(FGamepadKeyNames::MotionController_Right_TriggerAxis, GetPath(Instance, "/user/hand/right/input/trigger"));
+	InteractionMappings.Add(FGamepadKeyNames::MotionController_Right_Grip1, GetPath(Instance, "/user/hand/right/input/grip"));
+	InteractionMappings.Add(FGamepadKeyNames::MotionController_Right_Grip1Axis, GetPath(Instance, "/user/hand/right/input/grip"));
 	InteractionMappings.Add(FGamepadKeyNames::MotionController_Right_Thumbstick_X, GetPath(Instance, "/user/hand/right/input/thumbstick/x"));
 	InteractionMappings.Add(FGamepadKeyNames::MotionController_Right_Thumbstick_Y, GetPath(Instance, "/user/hand/right/input/thumbstick/y"));
 	InteractionMappings.Add(FGamepadKeyNames::MotionController_Right_Thumbstick, GetPath(Instance, "/user/hand/right/input/thumbstick/click"));
