Fix bad merge on OpenXR 1.0 plugin.

diff --git a/Engine/Plugins/Runtime/OpenXR/Source/OpenXRHMD/Private/OpenXRHMD.h b/Engine/Plugins/Runtime/OpenXR/Source/OpenXRHMD/Private/OpenXRHMD.h
--- a/Engine/Plugins/Runtime/OpenXR/Source/OpenXRHMD/Private/OpenXRHMD.h
+++ b/Engine/Plugins/Runtime/OpenXR/Source/OpenXRHMD/Private/OpenXRHMD.h
@@ -75,6 +75,15 @@ public:
 		return (TrackingSpaceType == XR_REFERENCE_SPACE_TYPE_STAGE) ? EHMDTrackingOrigin::Floor : EHMDTrackingOrigin::Eye;
 	}
 
+	virtual class IHeadMountedDisplay* GetHMDDevice() override
+	{
+		return this;
+	}
+	virtual class TSharedPtr< class IStereoRendering, ESPMode::ThreadSafe > GetStereoRenderingDevice() override
+	{
+		return SharedThis(this);
+	}
+
 protected:
 	/** FXRTrackingSystemBase protected interface */
 	virtual float GetWorldToMetersScale() const override;
