Copying latest OpenXR plugin changes from //UE4/Main.

#rnx

diff --git a/Engine/Plugins/Runtime/OpenXR/Source/OpenXRHMD/OpenXRHMD.Build.cs b/Engine/Plugins/Runtime/OpenXR/Source/OpenXRHMD/OpenXRHMD.Build.cs
--- a/Engine/Plugins/Runtime/OpenXR/Source/OpenXRHMD/OpenXRHMD.Build.cs
+++ b/Engine/Plugins/Runtime/OpenXR/Source/OpenXRHMD/OpenXRHMD.Build.cs
@@ -36,7 +36,8 @@ namespace UnrealBuildTool.Rules
 					"Core",
 					"CoreUObject",
 					"Engine",
-					"InputCore",
+                    "BuildSettings",
+                    "InputCore",
 					"RHI",
 					"RenderCore",
 					"Renderer",
diff --git a/Engine/Plugins/Runtime/OpenXR/Source/OpenXRHMD/Private/OpenXRHMD.cpp b/Engine/Plugins/Runtime/OpenXR/Source/OpenXRHMD/Private/OpenXRHMD.cpp
--- a/Engine/Plugins/Runtime/OpenXR/Source/OpenXRHMD/Private/OpenXRHMD.cpp
+++ b/Engine/Plugins/Runtime/OpenXR/Source/OpenXRHMD/Private/OpenXRHMD.cpp
@@ -22,7 +22,7 @@
 #include "PipelineStateCache.h"
 #include "Slate/SceneViewport.h"
 #include "Engine/GameEngine.h"
-#include "HardwareInfo.h"
+#include "BuildSettings.h"
 
 #if WITH_EDITOR
 #include "Editor/UnrealEd/Classes/Editor/EditorEngine.h"
@@ -120,7 +120,7 @@ private:
 	TSet<FString> AvailableExtensions;
 	TRefCountPtr<FOpenXRRenderBridge> RenderBridge;
 
-	void EnumerateExtensions();
+	bool EnumerateExtensions();
 	bool InitRenderBridge();
 };
 
@@ -156,21 +156,31 @@ uint64 FOpenXRHMDPlugin::GetGraphicsAdapterLuid()
 	return RenderBridge->GetGraphicsAdapterLuid();
 }
 
-void FOpenXRHMDPlugin::EnumerateExtensions()
+bool FOpenXRHMDPlugin::EnumerateExtensions()
 {
 	uint32_t ExtensionsCount = 0;
-	XR_ENSURE(xrEnumerateInstanceExtensionProperties(nullptr, 0, &ExtensionsCount, nullptr));
+	if (XR_FAILED(xrEnumerateInstanceExtensionProperties(nullptr, 0, &ExtensionsCount, nullptr)))
+	{
+		// If it fails this early that means there's no runtime installed
+		return false;
+	}
+
 	TArray<XrExtensionProperties> Properties;
 	Properties.SetNum(ExtensionsCount);
 	for (auto& Prop : Properties)
 	{
 		Prop = XrExtensionProperties{ XR_TYPE_EXTENSION_PROPERTIES };
 	}
-	XR_ENSURE(xrEnumerateInstanceExtensionProperties(nullptr, ExtensionsCount, &ExtensionsCount, Properties.GetData()));
-	for (const XrExtensionProperties& Prop : Properties)
+
+	if (XR_ENSURE(xrEnumerateInstanceExtensionProperties(nullptr, ExtensionsCount, &ExtensionsCount, Properties.GetData())))
 	{
-		AvailableExtensions.Add(Prop.extensionName);
+		for (const XrExtensionProperties& Prop : Properties)
+		{
+			AvailableExtensions.Add(Prop.extensionName);
+		}
+		return true;
 	}
+	return false;
 }
 
 bool FOpenXRHMDPlugin::InitRenderBridge()
@@ -237,19 +247,13 @@ bool FOpenXRHMDPlugin::PreInit()
 		return false;
 	}
 
-	// Engine registration can be disabled via console var.
-	auto* CVarDisableEngineAndAppRegistration = IConsoleManager::Get().FindTConsoleVariableDataInt(TEXT("r.DisableEngineAndAppRegistration"));
-	bool bDisableEngineRegistration = (CVarDisableEngineAndAppRegistration && CVarDisableEngineAndAppRegistration->GetValueOnAnyThread() != 0);
-
-	// EngineName will be of the form "UnrealEngine4.21", with the minor version ("21" in this example)
-	// updated with every quarterly release
-	FString EngineName = bDisableEngineRegistration ? FString("") : FApp::GetEpicProductIdentifier() + FEngineVersion::Current().ToString(EVersionComponent::Minor);
-	const TCHAR* ProjectName = bDisableEngineRegistration ? TEXT("") : FApp::GetProjectName();
-
-	EnumerateExtensions();
-
 	TArray<const char*> Extensions;
-	EnumerateExtensions();
+	if (!EnumerateExtensions())
+	{
+		UE_LOG(LogHMD, Log, TEXT("Failed to enumerate extensions. Please check if you have an OpenXR runtime installed."));
+		return false;
+	}
+
 #ifdef XR_USE_GRAPHICS_API_D3D11
 	if (HasExtension(XR_KHR_D3D11_ENABLE_EXTENSION_NAME))
 	{
@@ -275,22 +279,39 @@ bool FOpenXRHMDPlugin::PreInit()
 	}
 #endif
 
+#if 0
 	if (HasExtension(XR_KHR_COMPOSITION_LAYER_DEPTH_EXTENSION_NAME))
 	{
 		Extensions.Add(XR_KHR_COMPOSITION_LAYER_DEPTH_EXTENSION_NAME);
 	}
+#endif
 
 	if (HasExtension(XR_VARJO_QUAD_VIEWS_EXTENSION_NAME))
 	{
 		Extensions.Add(XR_VARJO_QUAD_VIEWS_EXTENSION_NAME);
 	}
 
+	// Engine registration can be disabled via console var.
+	auto* CVarDisableEngineAndAppRegistration = IConsoleManager::Get().FindTConsoleVariableDataInt(TEXT("r.DisableEngineAndAppRegistration"));
+	bool bDisableEngineRegistration = (CVarDisableEngineAndAppRegistration && CVarDisableEngineAndAppRegistration->GetValueOnAnyThread() != 0);
+
+	FText ProjectName = FText();
+	GConfig->GetText(TEXT("/Script/EngineSettings.GeneralProjectSettings"), TEXT("ProjectName"), ProjectName, GGameIni);
+
+	FText ProjectVersion = FText();
+	GConfig->GetText(TEXT("/Script/EngineSettings.GeneralProjectSettings"), TEXT("ProjectVersion"), ProjectVersion, GGameIni);
+
+	// EngineName will be of the form "UnrealEngine4.21", with the minor version ("21" in this example)
+	// updated with every quarterly release
+	FString EngineName = bDisableEngineRegistration ? FString("") : FApp::GetEpicProductIdentifier() + FEngineVersion::Current().ToString(EVersionComponent::Minor);
+	FString AppName = bDisableEngineRegistration ? TEXT("") : ProjectName.ToString() + ProjectVersion.ToString();
+
 	XrInstanceCreateInfo Info;
 	Info.type = XR_TYPE_INSTANCE_CREATE_INFO;
 	Info.next = nullptr;
 	Info.createFlags = 0;
-	FTCHARToUTF8_Convert::Convert(Info.applicationInfo.applicationName, XR_MAX_APPLICATION_NAME_SIZE, ProjectName, TCString<TCHAR>::Strlen(ProjectName) + 1);
-	Info.applicationInfo.applicationVersion = FEngineVersion::Current().GetChangelist();
+	FTCHARToUTF8_Convert::Convert(Info.applicationInfo.applicationName, XR_MAX_APPLICATION_NAME_SIZE, *AppName, AppName.Len() + 1);
+	Info.applicationInfo.applicationVersion = static_cast<uint32>(BuildSettings::GetCurrentChangelist()) | (BuildSettings::IsLicenseeVersion() ? 0x80000000 : 0);
 	FTCHARToUTF8_Convert::Convert(Info.applicationInfo.engineName, XR_MAX_ENGINE_NAME_SIZE, *EngineName, EngineName.Len() + 1);
 	Info.applicationInfo.engineVersion = (uint32)(FEngineVersion::Current().GetMinor() << 16 | FEngineVersion::Current().GetPatch());
 	Info.applicationInfo.apiVersion = XR_CURRENT_API_VERSION;
@@ -453,7 +474,7 @@ bool FOpenXRHMD::EnumerateTrackedDevices(TArray<int32>& OutDevices, EXRTrackedDe
 	}
 	if (Type == EXRTrackedDeviceType::Any || Type == EXRTrackedDeviceType::Controller)
 	{
-		for (int32 i = 0; i < DeviceSpaces.Num(); i++)
+		for (int32 i = 0; i < ActionSpaces.Num(); i++)
 		{
 			OutDevices.Add(i);
 		}
@@ -475,14 +496,14 @@ bool FOpenXRHMD::GetCurrentPose(int32 DeviceId, FQuat& CurrentOrientation, FVect
 	CurrentOrientation = FQuat::Identity;
 	CurrentPosition = FVector::ZeroVector;
 
-	if (!DeviceSpaces.IsValidIndex(DeviceId) || FrameState.predictedDisplayTime <= 0)
+	if (!ActionSpaces.IsValidIndex(DeviceId) || !ActionSpaces[DeviceId].Space || FrameState.predictedDisplayTime <= 0)
 	{
 		return false;
 	}
 
 	XrSpaceLocation Location = {};
 	Location.type = XR_TYPE_SPACE_LOCATION;
-	XrResult Result = xrLocateSpace(DeviceSpaces[DeviceId], GetTrackingSpace(), FrameState.predictedDisplayTime, &Location);
+	XrResult Result = xrLocateSpace(ActionSpaces[DeviceId].Space, GetTrackingSpace(), FrameState.predictedDisplayTime, &Location);
 	if (!XR_ENSURE(Result))
 	{
 		return false;
@@ -540,11 +561,17 @@ FQuat FOpenXRHMD::GetBaseOrientation() const
 
 bool FOpenXRHMD::IsStereoEnabled() const
 {
-	return true;
+	return bStereoEnabled;
 }
 
 bool FOpenXRHMD::EnableStereo(bool stereo)
 {
+	if (stereo == bStereoEnabled)
+	{
+		return true;
+	}
+
+	bStereoEnabled = stereo;
 	if (stereo)
 	{
 		GEngine->bForceDisableFrameRateSmoothing = true;
@@ -698,14 +725,15 @@ bool FOpenXRHMD::IsActiveThisFrame(class FViewport* InViewport) const
 FOpenXRHMD::FOpenXRHMD(const FAutoRegister& AutoRegister, XrInstance InInstance, XrSystemId InSystem, TRefCountPtr<FOpenXRRenderBridge>& InRenderBridge)
 	: FHeadMountedDisplayBase(nullptr)
 	, FSceneViewExtensionBase(AutoRegister)
+	, bStereoEnabled(false)
 	, bIsRunning(false)
 	, bIsReady(false)
+	, bIsRendering(false)
 	, bRunRequested(false)
 	, CurrentSessionState(XR_SESSION_STATE_UNKNOWN)
 	, Instance(InInstance)
 	, System(InSystem)
 	, Session(XR_NULL_HANDLE)
-	, DeviceSpaces()
 	, LocalSpace(XR_NULL_HANDLE)
 	, StageSpace(XR_NULL_HANDLE)
 	, TrackingSpaceType(XR_REFERENCE_SPACE_TYPE_STAGE)
@@ -758,6 +786,9 @@ FOpenXRHMD::FOpenXRHMD(const FAutoRegister& AutoRegister, XrInstance InInstance,
 		View.fov = XrFovf{ -PI / 4.0f, PI / 4.0f, PI / 4.0f, -PI / 4.0f };
 		View.pose = ToXrPose(FTransform::Identity);
 	}
+
+	// The HMD device does not have an action associated with it
+	ensure(ActionSpaces.Emplace(XR_NULL_HANDLE) == HMDDeviceId);
 }
 
 bool FOpenXRHMD::OnStereoStartup()
@@ -783,7 +814,7 @@ bool FOpenXRHMD::OnStereoStartup()
 	XR_ENSURE(xrEnumerateReferenceSpaces(Session, (uint32_t)spaces.Num(), &referenceSpacesCount, spaces.GetData()));
 	ensure(referenceSpacesCount == spaces.Num());
 
-	XrSpace Space;
+	XrSpace HmdSpace = XR_NULL_HANDLE;
 	XrReferenceSpaceCreateInfo SpaceInfo;
 
 	ensure(spaces.Contains(XR_REFERENCE_SPACE_TYPE_VIEW));
@@ -791,8 +822,8 @@ bool FOpenXRHMD::OnStereoStartup()
 	SpaceInfo.next = nullptr;
 	SpaceInfo.referenceSpaceType = XR_REFERENCE_SPACE_TYPE_VIEW;
 	SpaceInfo.poseInReferenceSpace = ToXrPose(FTransform::Identity);
-	XR_ENSURE(xrCreateReferenceSpace(Session, &SpaceInfo, &Space));
-	ensure(DeviceSpaces.Add(Space) == (int)DeviceSpaceHMD);
+	XR_ENSURE(xrCreateReferenceSpace(Session, &SpaceInfo, &HmdSpace));
+	ActionSpaces[HMDDeviceId].Space = HmdSpace;
 
 	ensure(spaces.Contains(XR_REFERENCE_SPACE_TYPE_LOCAL));
 	SpaceInfo.referenceSpaceType = XR_REFERENCE_SPACE_TYPE_LOCAL;
@@ -810,17 +841,10 @@ bool FOpenXRHMD::OnStereoStartup()
 		TrackingSpaceType = XR_REFERENCE_SPACE_TYPE_LOCAL;
 	}
 
-	XrActionSpaceCreateInfo ActionSpaceInfo;
-	ActionSpaceInfo.type = XR_TYPE_ACTION_SPACE_CREATE_INFO;
-	ActionSpaceInfo.next = nullptr;
-	ActionSpaceInfo.subactionPath = XR_NULL_PATH;
-	ActionSpaceInfo.poseInActionSpace = ToXrPose(FTransform::Identity);
-
-	for (auto && Action : DeferredActionDevices)
+	// Create action spaces for all devices
+	for (auto& ActionSpace : ActionSpaces)
 	{
-		ActionSpaceInfo.action = Action;
-		XR_ENSURE(xrCreateActionSpace(Session, &ActionSpaceInfo, &Space));
-		DeviceSpaces.Add(Space);
+		ActionSpace.CreateSpace(Session);
 	}
 
 	RenderBridge->SetOpenXRHMD(this);
@@ -856,7 +880,13 @@ FOpenXRHMD::~FOpenXRHMD()
 
 int32 FOpenXRHMD::AddActionDevice(XrAction Action)
 {
-	return DeferredActionDevices.Add(Action) + (int)KnownDeviceSpaceCount;
+	int32 DeviceId = ActionSpaces.Emplace(Action);
+	if (Session)
+	{
+		ActionSpaces[DeviceId].CreateSpace(Session);
+	}
+
+	return DeviceId;
 }
 
 bool FOpenXRHMD::IsInitialized() const
@@ -869,7 +899,7 @@ bool FOpenXRHMD::IsRunning() const
 	return bIsRunning;
 }
 
-bool FOpenXRHMD::IsRendering() const
+bool FOpenXRHMD::IsFocused() const
 {
 	return CurrentSessionState == XR_SESSION_STATE_FOCUSED;
 }
@@ -920,16 +950,18 @@ bool FOpenXRHMD::AllocateRenderTargetTexture(uint32 Index, uint32 SizeX, uint32
 	
 	// Grab the presentation texture out of the swapchain.
 	OutTargetableTexture = OutShaderResourceTexture = (FTexture2DRHIRef&)Swapchain->GetTextureRef();
-
+#if 0
 	// Allocate the depth buffer swapchain while we're here.
 	DepthSwapchain = RenderBridge->CreateSwapchain(Session, PF_DepthStencil, SizeX, SizeY, NumMips, NumSamples, 0, TexCreate_DepthStencilTargetable);
 	if (!DepthSwapchain)
 	{
 		return false;
 	}
+#endif
 	return true;
 }
 
+#if 0
 bool FOpenXRHMD::AllocateDepthTexture(uint32 Index, uint32 SizeX, uint32 SizeY, uint8 Format, uint32 NumMips, uint32 InTexFlags, uint32 TargetableTextureFlags, FTexture2DRHIRef& OutTargetableTexture, FTexture2DRHIRef& OutShaderResourceTexture, uint32 NumSamples)
 {
 	if (!DepthSwapchain || DepthSwapchain->GetTexture2D()->GetSizeX() != SizeX || DepthSwapchain->GetTexture2D()->GetSizeY() != SizeY)
@@ -941,6 +973,7 @@ bool FOpenXRHMD::AllocateDepthTexture(uint32 Index, uint32 SizeX, uint32 SizeY,
 
 	return true;
 }
+#endif
 
 void FOpenXRHMD::OnBeginRendering_RenderThread(FRHICommandListImmediate& RHICmdList, FSceneViewFamily& ViewFamily)
 {
@@ -950,15 +983,20 @@ void FOpenXRHMD::OnBeginRendering_RenderThread(FRHICommandListImmediate& RHICmdL
 		BeginInfo.type = XR_TYPE_FRAME_BEGIN_INFO;
 		BeginInfo.next = nullptr;
 		xrBeginFrame(Session, &BeginInfo);
+		bIsRendering = true;
+
+		Swapchain->IncrementSwapChainIndex_RHIThread(FrameStateRHI.predictedDisplayPeriod);
+#if 0
+		DepthSwapchain->IncrementSwapChainIndex_RHIThread(FrameStateRHI.predictedDisplayPeriod);
+#endif
 	}
 
+	FrameStateRHI = FrameState;
+
 	const FSceneView* MainView = ViewFamily.Views[0];
 	check(MainView);
 	BaseTransform = FTransform(MainView->BaseHmdOrientation, MainView->BaseHmdLocation);
 
-	Swapchain->IncrementSwapChainIndex_RHIThread(FrameStateRHI.predictedDisplayPeriod);
-	DepthSwapchain->IncrementSwapChainIndex_RHIThread(FrameStateRHI.predictedDisplayPeriod);
-
 	ViewsRHI.SetNum(Views.Num());
 	DepthLayersRHI.SetNum(Views.Num());
 
@@ -990,6 +1028,7 @@ void FOpenXRHMD::OnBeginRendering_RenderThread(FRHICommandListImmediate& RHICmdL
 			}
 		};
 
+#if 0
 		DepthLayer.type = XR_TYPE_COMPOSITION_LAYER_DEPTH_INFO_KHR;
 		DepthLayer.next = nullptr;
 		DepthLayer.subImage.swapchain = static_cast<FOpenXRSwapchain*>(GetDepthSwapchain())->GetHandle();
@@ -999,12 +1038,12 @@ void FOpenXRHMD::OnBeginRendering_RenderThread(FRHICommandListImmediate& RHICmdL
 		DepthLayer.maxDepth = 0.0f;
 		DepthLayer.nearZ = NearZ;
 		DepthLayer.farZ = FLT_MAX;
+#endif
 
 		OffsetX += Config.recommendedImageRectWidth;
 	}
 
 	// Give the RHI thread its own copy of the frame state and tracking space
-	FrameStateRHI = FrameState;
 	TrackingSpaceRHI = GetTrackingSpace();
 }
 
@@ -1035,7 +1074,7 @@ void FOpenXRHMD::OnBeginRendering_GameThread()
 	// Pass in a clean FrameState (some implementations may fail if this is just uninitialized).
 	FrameState.type = XR_TYPE_FRAME_STATE;
 	FrameState.next = nullptr;
-	FrameState.predictedDisplayPeriod = FrameStateRHI.predictedDisplayTime = 0;
+	FrameState.predictedDisplayPeriod = 0;
 
 	XR_ENSURE(xrWaitFrame(Session, &WaitInfo, &FrameState));
 
@@ -1044,21 +1083,28 @@ void FOpenXRHMD::OnBeginRendering_GameThread()
 	ViewInfo.type = XR_TYPE_VIEW_LOCATE_INFO;
 	ViewInfo.next = nullptr;
 	ViewInfo.viewConfigurationType = XR_VIEW_CONFIGURATION_TYPE_PRIMARY_STEREO;
-	ViewInfo.space = DeviceSpaces[HMDDeviceId];
+	ViewInfo.space = ActionSpaces[HMDDeviceId].Space;
 	ViewInfo.displayTime = FrameState.predictedDisplayTime;
 	XR_ENSURE(xrLocateViews(Session, &ViewInfo, &ViewState, 0, &ViewCount, nullptr));
 	Views.SetNum(ViewCount);
 	XR_ENSURE(xrLocateViews(Session, &ViewInfo, &ViewState, Views.Num(), &ViewCount, Views.GetData()));
 }
 
+bool FOpenXRHMD::ReadNextEvent(XrEventDataBuffer* buffer)
+{
+	// It is sufficient to clear just the XrEventDataBuffer header to XR_TYPE_EVENT_DATA_BUFFER
+	XrEventDataBaseHeader* baseHeader = reinterpret_cast<XrEventDataBaseHeader*>(buffer);
+	*baseHeader = { XR_TYPE_EVENT_DATA_BUFFER };
+	const XrResult xr = xrPollEvent(Instance, buffer);
+	XR_ENSURE(xr);
+	return xr == XR_SUCCESS;
+}
+
 bool FOpenXRHMD::OnStartGameFrame(FWorldContext& WorldContext)
 {
-	// Initialize an event buffer to hold the output.
+	// Process all pending messages.
 	XrEventDataBuffer event;
-	// Only the header needs to be initialized.
-	event.type = XR_TYPE_EVENT_DATA_BUFFER;
-	event.next = nullptr;
-	while (xrPollEvent(Instance, &event) == XR_SUCCESS) 
+	while (ReadNextEvent(&event))
 	{
 		switch (event.type) 
 		{
@@ -1069,7 +1115,7 @@ bool FOpenXRHMD::OnStartGameFrame(FWorldContext& WorldContext)
 
 			CurrentSessionState = SessionState.state;
 
-#if 1
+#if 0
 			static const TCHAR* StateTextMap[] =
 			{
 				TEXT("XR_SESSION_STATE_UNKNOWN"),
@@ -1123,6 +1169,12 @@ bool FOpenXRHMD::OnStartGameFrame(FWorldContext& WorldContext)
 
 			if (Session != XR_NULL_HANDLE)
 			{
+				// Clear up action spaces
+				for (auto& ActionSpace : ActionSpaces)
+				{
+					ActionSpace.DestroySpace();
+				}
+
 				// Close the session now we're allowed to.
 				ENQUEUE_RENDER_COMMAND(OpenXRDestroySession)([this](FRHICommandListImmediate& RHICmdList)
 				{
@@ -1134,18 +1186,14 @@ bool FOpenXRHMD::OnStartGameFrame(FWorldContext& WorldContext)
 					XR_ENSURE(xrDestroySession(Session));
 
 					Session = XR_NULL_HANDLE;
-
-					Swapchain = nullptr;
-					DepthSwapchain = nullptr;
 				});
 
 				FlushRenderingCommands();
 
+				bIsRendering = false;
+				bIsReady = false;
 				bIsRunning = false;
 				bRunRequested = false;
-
-				// Clear up device spaces
-				DeviceSpaces.Reset();
 			}
 
 			break;
@@ -1175,11 +1223,13 @@ void FOpenXRHMD::FinishRendering()
 	Layer.viewCount = ViewsRHI.Num();
 	Layer.views = ViewsRHI.GetData();
 
-	Swapchain->ReleaseCurrentImage_RHIThread();
-	DepthSwapchain->ReleaseCurrentImage_RHIThread();
-
-	if (bIsReady)
+	if (bIsRendering)
 	{
+		Swapchain->ReleaseCurrentImage_RHIThread();
+#if 0
+		DepthSwapchain->ReleaseCurrentImage_RHIThread();
+#endif
+
 		XrFrameEndInfo EndInfo;
 		XrCompositionLayerBaseHeader* Headers[1] = { reinterpret_cast<XrCompositionLayerBaseHeader*>(&Layer) };
 		EndInfo.type = XR_TYPE_FRAME_END_INFO;
@@ -1321,3 +1371,38 @@ void FOpenXRHMD::DrawHiddenAreaMesh_RenderThread(class FRHICommandList& RHICmdLi
 	RHICmdList.DrawIndexedPrimitive(Mesh.IndexBufferRHI, 0, 0, Mesh.NumVertices, 0, Mesh.NumTriangles, 1);
 #endif
 }
+
+//---------------------------------------------------
+// OpenXR Action Space Implementation
+//---------------------------------------------------
+
+FOpenXRHMD::FActionSpace::FActionSpace(XrAction InAction)
+	: Action(InAction)
+	, Space(XR_NULL_HANDLE)
+{
+}
+
+bool FOpenXRHMD::FActionSpace::CreateSpace(XrSession InSession)
+{
+	if (Action == XR_NULL_HANDLE || Space != XR_NULL_HANDLE)
+	{
+		return false;
+	}
+
+	XrActionSpaceCreateInfo ActionSpaceInfo;
+	ActionSpaceInfo.type = XR_TYPE_ACTION_SPACE_CREATE_INFO;
+	ActionSpaceInfo.next = nullptr;
+	ActionSpaceInfo.subactionPath = XR_NULL_PATH;
+	ActionSpaceInfo.poseInActionSpace = ToXrPose(FTransform::Identity);
+	ActionSpaceInfo.action = Action;
+	return XR_ENSURE(xrCreateActionSpace(InSession, &ActionSpaceInfo, &Space));
+}
+
+void FOpenXRHMD::FActionSpace::DestroySpace()
+{
+	if (Space)
+	{
+		XR_ENSURE(xrDestroySpace(Space));
+	}
+	Space = XR_NULL_HANDLE;
+}
diff --git a/Engine/Plugins/Runtime/OpenXR/Source/OpenXRHMD/Private/OpenXRHMD.h b/Engine/Plugins/Runtime/OpenXR/Source/OpenXRHMD/Private/OpenXRHMD.h
--- a/Engine/Plugins/Runtime/OpenXR/Source/OpenXRHMD/Private/OpenXRHMD.h
+++ b/Engine/Plugins/Runtime/OpenXR/Source/OpenXRHMD/Private/OpenXRHMD.h
@@ -26,6 +26,18 @@ class FOpenXRRenderBridge;
 class FOpenXRHMD : public FHeadMountedDisplayBase, public FXRRenderTargetManager, public FSceneViewExtensionBase
 {
 public:
+	class FActionSpace
+	{
+	public:
+		FActionSpace(XrAction InAction);
+
+		bool CreateSpace(XrSession InSession);
+		void DestroySpace();
+
+		XrAction Action;
+		XrSpace Space;
+	};
+
 	class FVulkanExtensions : public IHeadMountedDisplayVulkanExtensions
 	{
 	public:
@@ -91,6 +103,7 @@ protected:
 	bool StartSession();
 	bool OnStereoStartup();
 	bool OnStereoTeardown();
+	bool ReadNextEvent(XrEventDataBuffer* buffer);
 
 public:
 	/** IHeadMountedDisplay interface */
@@ -136,8 +149,10 @@ public:
 	/** IStereoRenderTargetManager */
 	virtual bool ShouldUseSeparateRenderTarget() const override { return true; }
 	virtual bool AllocateRenderTargetTexture(uint32 Index, uint32 SizeX, uint32 SizeY, uint8 Format, uint32 NumMips, uint32 Flags, uint32 TargetableTextureFlags, FTexture2DRHIRef& OutTargetableTexture, FTexture2DRHIRef& OutShaderResourceTexture, uint32 NumSamples = 1) override;
+#if 0
 	virtual bool NeedReAllocateDepthTexture(const TRefCountPtr<IPooledRenderTarget>& DepthTarget) override final { return false; }
 	virtual bool AllocateDepthTexture(uint32 Index, uint32 SizeX, uint32 SizeY, uint8 Format, uint32 NumMips, uint32 InTexFlags, uint32 TargetableTextureFlags, FTexture2DRHIRef& OutTargetableTexture, FTexture2DRHIRef& OutShaderResourceTexture, uint32 NumSamples = 1) override final;
+#endif
 
 	virtual FXRRenderBridge* GetActiveRenderBridge_GameThread(bool bUseSeparateRenderTarget) override;
 
@@ -155,13 +170,15 @@ public:
 	/** @return	True if the HMD was initialized OK */
 	OPENXRHMD_API bool IsInitialized() const;
 	OPENXRHMD_API bool IsRunning() const;
-	OPENXRHMD_API bool IsRendering() const;
+	OPENXRHMD_API bool IsFocused() const;
 	void FinishRendering();
 
 	OPENXRHMD_API int32 AddActionDevice(XrAction Action);
 
 	FXRSwapChain* GetSwapchain() { return Swapchain.Get(); }
+#if 0
 	FXRSwapChain* GetDepthSwapchain() { return DepthSwapchain.Get(); }
+#endif
 	XrInstance GetInstance() { return Instance; }
 	XrSystemId GetSystem() { return System; }
 	XrSession GetSession() { return Session; }
@@ -171,8 +188,10 @@ public:
 	}
 
 private:
+	bool					bStereoEnabled;
 	bool					bIsRunning;
 	bool					bIsReady;
+	bool					bIsRendering;
 	bool					bRunRequested;
 	
 	XrSessionState			CurrentSessionState;
@@ -181,7 +200,6 @@ private:
 	XrInstance				Instance;
 	XrSystemId				System;
 	XrSession				Session;
-	TArray<XrSpace>			DeviceSpaces;
 	XrSpace					LocalSpace;
 	XrSpace					StageSpace;
 	XrSpace					TrackingSpaceRHI;
@@ -199,19 +217,13 @@ private:
 #if 0
 	TArray<FHMDViewMesh>	HiddenAreaMeshes;
 #endif
+	TArray<FActionSpace>	ActionSpaces;
 
 	TRefCountPtr<FOpenXRRenderBridge> RenderBridge;
 	IRendererModule*		RendererModule;
 
 	FXRSwapChainPtr			Swapchain;
+#if 0
 	FXRSwapChainPtr			DepthSwapchain;
-
-	enum KnownDeviceSpaces
-	{
-		DeviceSpaceHMD = 0,
-		// @todo: If we add more known (fixed) spaces, add an enum here.
-		KnownDeviceSpaceCount,
-	};
-
-	TArray<XrAction>		DeferredActionDevices;
+#endif
 };
diff --git a/Engine/Plugins/Runtime/OpenXR/Source/OpenXRHMD/Private/OpenXRHMDPrivate.h b/Engine/Plugins/Runtime/OpenXR/Source/OpenXRHMD/Private/OpenXRHMDPrivate.h
--- a/Engine/Plugins/Runtime/OpenXR/Source/OpenXRHMD/Private/OpenXRHMDPrivate.h
+++ b/Engine/Plugins/Runtime/OpenXR/Source/OpenXRHMD/Private/OpenXRHMDPrivate.h
@@ -24,12 +24,14 @@ constexpr const TCHAR* xrResultToString(XrResult e)
 	}
 }
 
-inline bool LogXrResult(XrResult Result)
-{
-	return ensureMsgf(XR_SUCCEEDED(Result), TEXT("OpenXR call failed with result %s"), xrResultToString(Result));
-}
-
-#define XR_ENSURE(x) LogXrResult(x)
+#if DO_CHECK
+#define XR_ENSURE(x) [] (XrResult Result) \
+	{ \
+		return ensureMsgf(XR_SUCCEEDED(Result), TEXT("OpenXR call failed with result %s"), xrResultToString(Result)); \
+	} (x)
+#else
+#define XR_ENSURE(x) XR_SUCCEEDED(x)
+#endif
 
 FORCEINLINE FQuat ToFQuat(XrQuaternionf Quat)
 {
diff --git a/Engine/Plugins/Runtime/OpenXR/Source/OpenXRHMD/Private/OpenXRHMD_RenderBridge.cpp b/Engine/Plugins/Runtime/OpenXR/Source/OpenXRHMD/Private/OpenXRHMD_RenderBridge.cpp
--- a/Engine/Plugins/Runtime/OpenXR/Source/OpenXRHMD/Private/OpenXRHMD_RenderBridge.cpp
+++ b/Engine/Plugins/Runtime/OpenXR/Source/OpenXRHMD/Private/OpenXRHMD_RenderBridge.cpp
@@ -100,15 +100,17 @@ public:
 		XrGraphicsRequirementsOpenGLKHR Requirements;
 		Requirements.type = XR_TYPE_GRAPHICS_REQUIREMENTS_OPENGL_KHR;
 		Requirements.next = nullptr;
+		Requirements.minApiVersionSupported = 0;
+		Requirements.maxApiVersionSupported = 0;
 		XR_ENSURE(xrGetOpenGLGraphicsRequirementsKHR(InInstance, InSystem, &Requirements));
 
 		XrVersion RHIVersion = XR_MAKE_VERSION(FOpenGL::GetMajorVersion(), FOpenGL::GetMinorVersion(), 0);
-		if (RHIVersion < Requirements.minApiVersionSupported)
+		if (RHIVersion < Requirements.minApiVersionSupported) //-V547
 		{
 			UE_LOG(LogHMD, Fatal, TEXT("The OpenGL API version does not meet the minimum version required by the OpenXR runtime"));
 		}
 
-		if (RHIVersion > Requirements.maxApiVersionSupported)
+		if (RHIVersion > Requirements.maxApiVersionSupported) //-V547
 		{
 			UE_LOG(LogHMD, Warning, TEXT("The OpenGL API version has not been tested with the OpenXR runtime"));
 		}
@@ -152,6 +154,8 @@ public:
 		XrGraphicsRequirementsVulkanKHR Requirements;
 		Requirements.type = XR_TYPE_GRAPHICS_REQUIREMENTS_VULKAN_KHR;
 		Requirements.next = nullptr;
+		Requirements.minApiVersionSupported = 0;
+		Requirements.maxApiVersionSupported = 0;
 		XR_ENSURE(xrGetVulkanGraphicsRequirementsKHR(InInstance, InSystem, &Requirements));
 
 		// The extension uses the OpenXR version format instead of the Vulkan one
@@ -160,12 +164,12 @@ public:
 			VK_VERSION_MINOR(UE_VK_API_VERSION),
 			VK_VERSION_PATCH(UE_VK_API_VERSION)
 		);
-		if (RHIVersion < Requirements.minApiVersionSupported)
+		if (RHIVersion < Requirements.minApiVersionSupported) //-V547
 		{
 			UE_LOG(LogHMD, Fatal, TEXT("The Vulkan API version does not meet the minimum version required by the OpenXR runtime"));
 		}
 
-		if (RHIVersion > Requirements.maxApiVersionSupported)
+		if (RHIVersion > Requirements.maxApiVersionSupported) //-V547
 		{
 			UE_LOG(LogHMD, Warning, TEXT("The Vulkan API version has not been tested with the OpenXR runtime"));
 		}
diff --git a/Engine/Plugins/Runtime/OpenXR/Source/OpenXRHMD/Private/OpenXRHMD_Swapchain.cpp b/Engine/Plugins/Runtime/OpenXR/Source/OpenXRHMD/Private/OpenXRHMD_Swapchain.cpp
--- a/Engine/Plugins/Runtime/OpenXR/Source/OpenXRHMD/Private/OpenXRHMD_Swapchain.cpp
+++ b/Engine/Plugins/Runtime/OpenXR/Source/OpenXRHMD/Private/OpenXRHMD_Swapchain.cpp
@@ -57,8 +57,13 @@ void FOpenXRSwapchain::ReleaseResources_RHIThread()
 	xrDestroySwapchain(Handle);
 }
 
-XrSwapchain CreateSwapchain(XrSession InSession, uint32 PlatformFormat, uint32 SizeX, uint32 SizeY, uint32 NumMips, uint32 NumSamples, uint32 TargetableTextureFlags)
+uint8 GetNearestSupportedSwapchainFormat(XrSession InSession, uint8 RequestedFormat, TFunction<uint32(uint8)> ToPlatformFormat = nullptr)
 {
+	if (!ToPlatformFormat)
+	{
+		ToPlatformFormat = [](uint8 InFormat) { return GPixelFormats[InFormat].PlatformFormat; };
+	}
+
 	uint32_t FormatsCount;
 	xrEnumerateSwapchainFormats(InSession, 0, &FormatsCount, nullptr);
 
@@ -67,12 +72,44 @@ XrSwapchain CreateSwapchain(XrSession InSession, uint32 PlatformFormat, uint32 S
 	XR_ENSURE(xrEnumerateSwapchainFormats(InSession, (uint32_t)Formats.Num(), &FormatsCount, Formats.GetData()));
 	ensure(FormatsCount == Formats.Num());
 
-	if (!Formats.Contains(PlatformFormat))
+	// Return immediately if the runtime supports the exact format being requested.
+	uint32 PlatformFormat = ToPlatformFormat(RequestedFormat);
+	if (Formats.Contains(PlatformFormat))
 	{
-		UE_LOG(LogHMD, Log, TEXT("Swapchain format not supported (%d), falling back to runtime preferred format (%d)."), PlatformFormat, Formats[0]);
-		PlatformFormat = Formats[0];
+		return RequestedFormat;
 	}
 
+	// Search for an 8bpc fallback format in order of preference (first element in the array has the high preference).
+	uint8 FallbackFormat = 0;
+	uint32 FallbackPlatformFormat = 0;
+	for (int64_t Format : Formats)
+	{
+		if (Format == ToPlatformFormat(PF_B8G8R8A8))
+		{
+			FallbackFormat = PF_B8G8R8A8;
+			FallbackPlatformFormat = Format;
+			break;
+		}
+		else if (Format == ToPlatformFormat(PF_R8G8B8A8))
+		{
+			FallbackFormat = PF_R8G8B8A8;
+			FallbackPlatformFormat = Format;
+			break;
+		}
+	}
+
+	if (!FallbackFormat)
+	{
+		UE_LOG(LogHMD, Warning, TEXT("No compatible swapchain format found!"));
+		return PF_Unknown;
+	}
+
+	UE_LOG(LogHMD, Warning, TEXT("Swapchain format not supported (%d), falling back to runtime preferred format (%d)."), PlatformFormat, FallbackPlatformFormat);
+	return FallbackFormat;
+}
+
+XrSwapchain CreateSwapchain(XrSession InSession, uint32 PlatformFormat, uint32 SizeX, uint32 SizeY, uint32 NumMips, uint32 NumSamples, uint32 TargetableTextureFlags)
+{
 	XrSwapchainUsageFlags Usage = 0;
 	if (TargetableTextureFlags & TexCreate_RenderTargetable)
 	{
@@ -129,11 +166,22 @@ TArray<T> EnumerateImages(XrSwapchain InSwapchain, XrStructureType InType)
 #ifdef XR_USE_GRAPHICS_API_D3D11
 FXRSwapChainPtr CreateSwapchain_D3D11(XrSession InSession, uint8 Format, uint32 SizeX, uint32 SizeY, uint32 NumMips, uint32 NumSamples, uint32 Flags, uint32 TargetableTextureFlags)
 {
-	// We need to convert typeless to typed formats to create a swapchain
-	DXGI_FORMAT PlatformFormat = (DXGI_FORMAT)GPixelFormats[Format].PlatformFormat;
-	PlatformFormat = FindDepthStencilDXGIFormat(PlatformFormat);
-	PlatformFormat = FindShaderResourceDXGIFormat(PlatformFormat, Flags & TexCreate_SRGB);
-	XrSwapchain Swapchain = CreateSwapchain(InSession, PlatformFormat, SizeX, SizeY, NumMips, NumSamples, TargetableTextureFlags);
+	TFunction<uint32(uint8)> ToPlatformFormat = [Flags](uint8 InFormat)
+	{
+		// We need to convert typeless to typed formats to create a swapchain
+		DXGI_FORMAT PlatformFormat = (DXGI_FORMAT)GPixelFormats[InFormat].PlatformFormat;
+		PlatformFormat = FindDepthStencilDXGIFormat(PlatformFormat);
+		PlatformFormat = FindShaderResourceDXGIFormat(PlatformFormat, Flags & TexCreate_SRGB);
+		return PlatformFormat;
+	};
+
+	Format = GetNearestSupportedSwapchainFormat(InSession, Format, ToPlatformFormat);
+	if (!Format)
+	{
+		return nullptr;
+	}
+
+	XrSwapchain Swapchain = CreateSwapchain(InSession, ToPlatformFormat(Format), SizeX, SizeY, NumMips, NumSamples, TargetableTextureFlags);
 	if (!Swapchain)
 	{
 		return nullptr;
@@ -156,11 +204,22 @@ FXRSwapChainPtr CreateSwapchain_D3D11(XrSession InSession, uint8 Format, uint32
 #ifdef XR_USE_GRAPHICS_API_D3D12
 FXRSwapChainPtr CreateSwapchain_D3D12(XrSession InSession, uint8 Format, uint32 SizeX, uint32 SizeY, uint32 NumMips, uint32 NumSamples, uint32 Flags, uint32 TargetableTextureFlags)
 {
-	// We need to convert typeless to typed formats to create a swapchain
-	DXGI_FORMAT PlatformFormat = (DXGI_FORMAT)GPixelFormats[Format].PlatformFormat;
-	PlatformFormat = FindDepthStencilDXGIFormat_D3D12(PlatformFormat);
-	PlatformFormat = FindShaderResourceDXGIFormat_D3D12(PlatformFormat, Flags & TexCreate_SRGB);
-	XrSwapchain Swapchain = CreateSwapchain(InSession, Format, SizeX, SizeY, NumMips, NumSamples, TargetableTextureFlags);
+	TFunction<uint32(uint8)> ToPlatformFormat = [Flags](uint8 InFormat)
+	{
+		// We need to convert typeless to typed formats to create a swapchain
+		DXGI_FORMAT PlatformFormat = (DXGI_FORMAT)GPixelFormats[InFormat].PlatformFormat;
+		PlatformFormat = FindDepthStencilDXGIFormat_D3D12(PlatformFormat);
+		PlatformFormat = FindShaderResourceDXGIFormat_D3D12(PlatformFormat, Flags & TexCreate_SRGB);
+		return PlatformFormat;
+	};
+
+	Format = GetNearestSupportedSwapchainFormat(InSession, Format, ToPlatformFormat);
+	if (!Format)
+	{
+		return nullptr;
+	}
+
+	XrSwapchain Swapchain = CreateSwapchain(InSession, ToPlatformFormat(Format), SizeX, SizeY, NumMips, NumSamples, TargetableTextureFlags);
 	if (!Swapchain)
 	{
 		return nullptr;
@@ -183,6 +242,12 @@ FXRSwapChainPtr CreateSwapchain_D3D12(XrSession InSession, uint8 Format, uint32
 #ifdef XR_USE_GRAPHICS_API_OPENGL
 FXRSwapChainPtr CreateSwapchain_OpenGL(XrSession InSession, uint8 Format, uint32 SizeX, uint32 SizeY, uint32 NumMips, uint32 NumSamples, uint32 Flags, uint32 TargetableTextureFlags)
 {
+	Format = GetNearestSupportedSwapchainFormat(InSession, Format);
+	if (!Format)
+	{
+		return nullptr;
+	}
+
 	XrSwapchain Swapchain = CreateSwapchain(InSession, GPixelFormats[Format].PlatformFormat, SizeX, SizeY, NumMips, NumSamples, TargetableTextureFlags);
 	if (!Swapchain)
 	{
@@ -206,6 +271,12 @@ FXRSwapChainPtr CreateSwapchain_OpenGL(XrSession InSession, uint8 Format, uint32
 #ifdef XR_USE_GRAPHICS_API_VULKAN
 FXRSwapChainPtr CreateSwapchain_Vulkan(XrSession InSession, uint8 Format, uint32 SizeX, uint32 SizeY, uint32 NumMips, uint32 NumSamples, uint32 Flags, uint32 TargetableTextureFlags)
 {
+	Format = GetNearestSupportedSwapchainFormat(InSession, Format);
+	if (!Format)
+	{
+		return nullptr;
+	}
+
 	XrSwapchain Swapchain = CreateSwapchain(InSession, GPixelFormats[Format].PlatformFormat, SizeX, SizeY, NumMips, NumSamples, TargetableTextureFlags);
 	if (!Swapchain)
 	{
diff --git a/Engine/Plugins/Runtime/OpenXR/Source/OpenXRInput/Private/OpenXRInput.cpp b/Engine/Plugins/Runtime/OpenXR/Source/OpenXRInput/Private/OpenXRInput.cpp
--- a/Engine/Plugins/Runtime/OpenXR/Source/OpenXRInput/Private/OpenXRInput.cpp
+++ b/Engine/Plugins/Runtime/OpenXR/Source/OpenXRInput/Private/OpenXRInput.cpp
@@ -107,7 +107,7 @@ FOpenXRInputPlugin::FOpenXRAction::FOpenXRAction(XrActionSet InSet, const FInput
 }
 
 FOpenXRInputPlugin::FOpenXRAction::FOpenXRAction(XrActionSet InSet, const FInputAxisKeyMapping& InAxisKey)
-	: FOpenXRAction(InSet, XR_ACTION_TYPE_VECTOR2F_INPUT, InAxisKey.AxisName)
+	: FOpenXRAction(InSet, XR_ACTION_TYPE_FLOAT_INPUT, InAxisKey.AxisName)
 {
 	ActionKey = InAxisKey.Key.GetFName();
 }
@@ -155,7 +155,7 @@ FOpenXRInputPlugin::FOpenXRInput::FOpenXRInput(FOpenXRHMD* HMD)
 	XrInstance Instance = OpenXRHMD->GetInstance();
 	check(Instance);
 
-	XrActionSet ActionSet;
+	XrActionSet ActionSet = XR_NULL_HANDLE;
 	XrActionSetCreateInfo SetInfo;
 	SetInfo.type = XR_TYPE_ACTION_SET_CREATE_INFO;
 	SetInfo.next = nullptr;
@@ -168,8 +168,8 @@ FOpenXRInputPlugin::FOpenXRInput::FOpenXRInput(FOpenXRHMD* HMD)
 	Controllers.Add(EControllerHand::Left, FOpenXRController(OpenXRHMD, ActionSet, "Left Controller"));
 	Controllers.Add(EControllerHand::Right, FOpenXRController(OpenXRHMD, ActionSet, "Right Controller"));
 
-	Bindings.Add(XrActionSuggestedBinding{ Controllers[EControllerHand::Left].Action, GetPath(Instance, "/user/hand/left/input/palm") });
-	Bindings.Add(XrActionSuggestedBinding{ Controllers[EControllerHand::Right].Action, GetPath(Instance, "/user/hand/right/input/palm") });
+	Bindings.Add(XrActionSuggestedBinding{ Controllers[EControllerHand::Left].Action, GetPath(Instance, "/user/hand/left/input/grip/pose") });
+	Bindings.Add(XrActionSuggestedBinding{ Controllers[EControllerHand::Right].Action, GetPath(Instance, "/user/hand/right/input/grip/pose") });
 
 	Bindings.Add(XrActionSuggestedBinding{ Controllers[EControllerHand::Left].VibrationAction, GetPath(Instance, "/user/hand/left/output/haptic") });
 	Bindings.Add(XrActionSuggestedBinding{ Controllers[EControllerHand::Right].VibrationAction, GetPath(Instance, "/user/hand/right/output/haptic") });
@@ -177,24 +177,24 @@ FOpenXRInputPlugin::FOpenXRInput::FOpenXRInput(FOpenXRHMD* HMD)
 	InteractionMappings.Add(FGamepadKeyNames::MotionController_Left_Shoulder, GetPath(Instance, "/user/hand/left/input/menu/click"));
 	InteractionMappings.Add(FGamepadKeyNames::MotionController_Left_Trigger, GetPath(Instance, "/user/hand/left/input/trigger"));
 	InteractionMappings.Add(FGamepadKeyNames::MotionController_Left_TriggerAxis, GetPath(Instance, "/user/hand/left/input/trigger"));
-	InteractionMappings.Add(FGamepadKeyNames::MotionController_Left_Grip1, GetPath(Instance, "/user/hand/left/input/grip"));
-	InteractionMappings.Add(FGamepadKeyNames::MotionController_Left_Grip1Axis, GetPath(Instance, "/user/hand/left/input/grip"));
+	InteractionMappings.Add(FGamepadKeyNames::MotionController_Left_Grip1, GetPath(Instance, "/user/hand/left/input/squeeze"));
+	InteractionMappings.Add(FGamepadKeyNames::MotionController_Left_Grip1Axis, GetPath(Instance, "/user/hand/left/input/squeeze"));
 	InteractionMappings.Add(FGamepadKeyNames::MotionController_Left_Thumbstick_X, GetPath(Instance, "/user/hand/left/input/thumbstick/x"));
 	InteractionMappings.Add(FGamepadKeyNames::MotionController_Left_Thumbstick_Y, GetPath(Instance, "/user/hand/left/input/thumbstick/y"));
 	InteractionMappings.Add(FGamepadKeyNames::MotionController_Left_Thumbstick, GetPath(Instance, "/user/hand/left/input/thumbstick/click"));
-	InteractionMappings.Add(FGamepadKeyNames::MotionController_Left_FaceButton1, GetPath(Instance, "/user/hand/left/input/x/click"));
-	InteractionMappings.Add(FGamepadKeyNames::MotionController_Left_FaceButton2, GetPath(Instance, "/user/hand/left/input/y/click"));
+	//InteractionMappings.Add(FGamepadKeyNames::MotionController_Left_FaceButton1, GetPath(Instance, "/user/hand/left/input/x/click"));
+	//InteractionMappings.Add(FGamepadKeyNames::MotionController_Left_FaceButton2, GetPath(Instance, "/user/hand/left/input/y/click"));
 
 	InteractionMappings.Add(FGamepadKeyNames::MotionController_Right_Shoulder, GetPath(Instance, "/user/hand/right/input/menu/click"));
 	InteractionMappings.Add(FGamepadKeyNames::MotionController_Right_Trigger, GetPath(Instance, "/user/hand/right/input/trigger"));
 	InteractionMappings.Add(FGamepadKeyNames::MotionController_Right_TriggerAxis, GetPath(Instance, "/user/hand/right/input/trigger"));
-	InteractionMappings.Add(FGamepadKeyNames::MotionController_Right_Grip1, GetPath(Instance, "/user/hand/right/input/grip"));
-	InteractionMappings.Add(FGamepadKeyNames::MotionController_Right_Grip1Axis, GetPath(Instance, "/user/hand/right/input/grip"));
+	InteractionMappings.Add(FGamepadKeyNames::MotionController_Right_Grip1, GetPath(Instance, "/user/hand/right/input/squeeze"));
+	InteractionMappings.Add(FGamepadKeyNames::MotionController_Right_Grip1Axis, GetPath(Instance, "/user/hand/right/input/squeeze"));
 	InteractionMappings.Add(FGamepadKeyNames::MotionController_Right_Thumbstick_X, GetPath(Instance, "/user/hand/right/input/thumbstick/x"));
 	InteractionMappings.Add(FGamepadKeyNames::MotionController_Right_Thumbstick_Y, GetPath(Instance, "/user/hand/right/input/thumbstick/y"));
 	InteractionMappings.Add(FGamepadKeyNames::MotionController_Right_Thumbstick, GetPath(Instance, "/user/hand/right/input/thumbstick/click"));
-	InteractionMappings.Add(FGamepadKeyNames::MotionController_Right_FaceButton1, GetPath(Instance, "/user/hand/right/input/a/click"));
-	InteractionMappings.Add(FGamepadKeyNames::MotionController_Right_FaceButton2, GetPath(Instance, "/user/hand/right/input/b/click"));
+	//InteractionMappings.Add(FGamepadKeyNames::MotionController_Right_FaceButton1, GetPath(Instance, "/user/hand/right/input/a/click"));
+	//InteractionMappings.Add(FGamepadKeyNames::MotionController_Right_FaceButton2, GetPath(Instance, "/user/hand/right/input/b/click"));
 
 	auto InputSettings = GetDefault<UInputSettings>();
 	if (InputSettings != nullptr)
@@ -228,10 +228,10 @@ FOpenXRInputPlugin::FOpenXRInput::FOpenXRInput(FOpenXRHMD* HMD)
 	}
 
 	TArray<XrPath> Profiles;
-	Profiles.Add(GetPath(Instance, "/interaction_profiles/khr/simple_controller"));
+	//Profiles.Add(GetPath(Instance, "/interaction_profiles/khr/simple_controller"));
 	Profiles.Add(GetPath(Instance, "/interaction_profiles/microsoft/motion_controller"));
 	Profiles.Add(GetPath(Instance, "/interaction_profiles/oculus/touch_controller"));
-	Profiles.Add(GetPath(Instance, "/interaction_profiles/valve/knuckles_controller"));
+	Profiles.Add(GetPath(Instance, "/interaction_profiles/valve/index_controller"));
 
 	for (XrPath Profile : Profiles)
 	{
@@ -298,9 +298,10 @@ void FOpenXRInputPlugin::FOpenXRInput::AddAction(XrActionSet ActionSet, const TA
 
 void FOpenXRInputPlugin::FOpenXRInput::Tick(float DeltaTime)
 {
-	if (OpenXRHMD->IsRunning() && OpenXRHMD->IsRendering())
+	XrSession Session = OpenXRHMD->GetSession();
+
+	if (OpenXRHMD->IsRunning())
 	{
-		XrSession Session = OpenXRHMD->GetSession();
 		if (!bActionsBound)
 		{
 			TArray<XrActionSet> BindActionSets;
@@ -312,29 +313,31 @@ void FOpenXRInputPlugin::FOpenXRInput::Tick(float DeltaTime)
 			SessionActionSetsAttachInfo.next = nullptr;
 			SessionActionSetsAttachInfo.countActionSets = BindActionSets.Num();
 			SessionActionSetsAttachInfo.actionSets = BindActionSets.GetData();
-			xrAttachSessionActionSets(Session, &SessionActionSetsAttachInfo);
+			XR_ENSURE(xrAttachSessionActionSets(Session, &SessionActionSetsAttachInfo));
 
 			bActionsBound = true;
 		}
+	}
+	else if (bActionsBound)
+	{
+		// If the session shut down, clean up.
+		bActionsBound = false;
+	}
 
+	if (OpenXRHMD->IsFocused())
+	{
 		XrActionsSyncInfo SyncInfo;
 		SyncInfo.type = XR_TYPE_ACTIONS_SYNC_INFO;
 		SyncInfo.next = nullptr;
 		SyncInfo.countActiveActionSets = ActionSets.Num();
 		SyncInfo.activeActionSets = ActionSets.GetData();
- 		XR_ENSURE(xrSyncActions(Session, &SyncInfo));
-	}
-	else if (bActionsBound && !OpenXRHMD->IsRunning())
-	{
-		// If the session shut down, clean up.
-		bActionsBound = false;
+		XR_ENSURE(xrSyncActions(Session, &SyncInfo));
 	}
-
 }
 
 void FOpenXRInputPlugin::FOpenXRInput::SendControllerEvents()
 {
-	if (!OpenXRHMD->IsRunning())
+	if (!bActionsBound)
 	{
 		return;
 	}
