Collection of fixes to get the OpenXR plugin in a shippable state
- The swapchain format fallback is now functional
- Depth composition layers have been disabled
- Corrected various input binding paths
- Refactored action spaces
- Updated copyright notice

diff --git a/Engine/Plugins/Runtime/OpenXR/Source/OpenXRHMD/Private/OpenXRHMD.cpp b/Engine/Plugins/Runtime/OpenXR/Source/OpenXRHMD/Private/OpenXRHMD.cpp
--- a/Engine/Plugins/Runtime/OpenXR/Source/OpenXRHMD/Private/OpenXRHMD.cpp
+++ b/Engine/Plugins/Runtime/OpenXR/Source/OpenXRHMD/Private/OpenXRHMD.cpp
@@ -275,10 +275,12 @@ bool FOpenXRHMDPlugin::PreInit()
 	}
 #endif
 
+#if 0
 	if (HasExtension(XR_KHR_COMPOSITION_LAYER_DEPTH_EXTENSION_NAME))
 	{
 		Extensions.Add(XR_KHR_COMPOSITION_LAYER_DEPTH_EXTENSION_NAME);
 	}
+#endif
 
 	if (HasExtension(XR_VARJO_QUAD_VIEWS_EXTENSION_NAME))
 	{
@@ -453,7 +455,7 @@ bool FOpenXRHMD::EnumerateTrackedDevices(TArray<int32>& OutDevices, EXRTrackedDe
 	}
 	if (Type == EXRTrackedDeviceType::Any || Type == EXRTrackedDeviceType::Controller)
 	{
-		for (int32 i = 0; i < DeviceSpaces.Num(); i++)
+		for (int32 i = 0; i < ActionSpaces.Num(); i++)
 		{
 			OutDevices.Add(i);
 		}
@@ -475,14 +477,14 @@ bool FOpenXRHMD::GetCurrentPose(int32 DeviceId, FQuat& CurrentOrientation, FVect
 	CurrentOrientation = FQuat::Identity;
 	CurrentPosition = FVector::ZeroVector;
 
-	if (!DeviceSpaces.IsValidIndex(DeviceId) || FrameState.predictedDisplayTime <= 0)
+	if (!ActionSpaces.IsValidIndex(DeviceId) || FrameState.predictedDisplayTime <= 0)
 	{
 		return false;
 	}
 
 	XrSpaceLocation Location = {};
 	Location.type = XR_TYPE_SPACE_LOCATION;
-	XrResult Result = xrLocateSpace(DeviceSpaces[DeviceId], GetTrackingSpace(), FrameState.predictedDisplayTime, &Location);
+	XrResult Result = xrLocateSpace(ActionSpaces[DeviceId].Space, GetTrackingSpace(), FrameState.predictedDisplayTime, &Location);
 	if (!XR_ENSURE(Result))
 	{
 		return false;
@@ -540,11 +542,17 @@ FQuat FOpenXRHMD::GetBaseOrientation() const
 
 bool FOpenXRHMD::IsStereoEnabled() const
 {
-	return true;
+	return bStereoEnabled;
 }
 
 bool FOpenXRHMD::EnableStereo(bool stereo)
 {
+	if (stereo == bStereoEnabled)
+	{
+		return true;
+	}
+
+	bStereoEnabled = stereo;
 	if (stereo)
 	{
 		GEngine->bForceDisableFrameRateSmoothing = true;
@@ -698,14 +706,15 @@ bool FOpenXRHMD::IsActiveThisFrame(class FViewport* InViewport) const
 FOpenXRHMD::FOpenXRHMD(const FAutoRegister& AutoRegister, XrInstance InInstance, XrSystemId InSystem, TRefCountPtr<FOpenXRRenderBridge>& InRenderBridge)
 	: FHeadMountedDisplayBase(nullptr)
 	, FSceneViewExtensionBase(AutoRegister)
+	, bStereoEnabled(false)
 	, bIsRunning(false)
 	, bIsReady(false)
+	, bIsRendering(false)
 	, bRunRequested(false)
 	, CurrentSessionState(XR_SESSION_STATE_UNKNOWN)
 	, Instance(InInstance)
 	, System(InSystem)
 	, Session(XR_NULL_HANDLE)
-	, DeviceSpaces()
 	, LocalSpace(XR_NULL_HANDLE)
 	, StageSpace(XR_NULL_HANDLE)
 	, TrackingSpaceType(XR_REFERENCE_SPACE_TYPE_STAGE)
@@ -758,6 +767,9 @@ FOpenXRHMD::FOpenXRHMD(const FAutoRegister& AutoRegister, XrInstance InInstance,
 		View.fov = XrFovf{ -PI / 4.0f, PI / 4.0f, PI / 4.0f, -PI / 4.0f };
 		View.pose = ToXrPose(FTransform::Identity);
 	}
+
+	// The HMD device does not have an action associated with it
+	ensure(ActionSpaces.Emplace(XR_NULL_HANDLE) == HMDDeviceId);
 }
 
 bool FOpenXRHMD::OnStereoStartup()
@@ -783,7 +795,7 @@ bool FOpenXRHMD::OnStereoStartup()
 	XR_ENSURE(xrEnumerateReferenceSpaces(Session, (uint32_t)spaces.Num(), &referenceSpacesCount, spaces.GetData()));
 	ensure(referenceSpacesCount == spaces.Num());
 
-	XrSpace Space;
+	XrSpace HmdSpace;
 	XrReferenceSpaceCreateInfo SpaceInfo;
 
 	ensure(spaces.Contains(XR_REFERENCE_SPACE_TYPE_VIEW));
@@ -791,8 +803,8 @@ bool FOpenXRHMD::OnStereoStartup()
 	SpaceInfo.next = nullptr;
 	SpaceInfo.referenceSpaceType = XR_REFERENCE_SPACE_TYPE_VIEW;
 	SpaceInfo.poseInReferenceSpace = ToXrPose(FTransform::Identity);
-	XR_ENSURE(xrCreateReferenceSpace(Session, &SpaceInfo, &Space));
-	ensure(DeviceSpaces.Add(Space) == (int)DeviceSpaceHMD);
+	XR_ENSURE(xrCreateReferenceSpace(Session, &SpaceInfo, &HmdSpace));
+	ActionSpaces[HMDDeviceId].Space = HmdSpace;
 
 	ensure(spaces.Contains(XR_REFERENCE_SPACE_TYPE_LOCAL));
 	SpaceInfo.referenceSpaceType = XR_REFERENCE_SPACE_TYPE_LOCAL;
@@ -810,17 +822,10 @@ bool FOpenXRHMD::OnStereoStartup()
 		TrackingSpaceType = XR_REFERENCE_SPACE_TYPE_LOCAL;
 	}
 
-	XrActionSpaceCreateInfo ActionSpaceInfo;
-	ActionSpaceInfo.type = XR_TYPE_ACTION_SPACE_CREATE_INFO;
-	ActionSpaceInfo.next = nullptr;
-	ActionSpaceInfo.subactionPath = XR_NULL_PATH;
-	ActionSpaceInfo.poseInActionSpace = ToXrPose(FTransform::Identity);
-
-	for (auto && Action : DeferredActionDevices)
+	// Create action spaces for all devices
+	for (auto& ActionSpace : ActionSpaces)
 	{
-		ActionSpaceInfo.action = Action;
-		XR_ENSURE(xrCreateActionSpace(Session, &ActionSpaceInfo, &Space));
-		DeviceSpaces.Add(Space);
+		ActionSpace.CreateSpace(Session);
 	}
 
 	RenderBridge->SetOpenXRHMD(this);
@@ -856,7 +861,13 @@ FOpenXRHMD::~FOpenXRHMD()
 
 int32 FOpenXRHMD::AddActionDevice(XrAction Action)
 {
-	return DeferredActionDevices.Add(Action) + (int)KnownDeviceSpaceCount;
+	int32 DeviceId = ActionSpaces.Emplace(Action);
+	if (Session)
+	{
+		ActionSpaces[DeviceId].CreateSpace(Session);
+	}
+
+	return DeviceId;
 }
 
 bool FOpenXRHMD::IsInitialized() const
@@ -920,16 +931,18 @@ bool FOpenXRHMD::AllocateRenderTargetTexture(uint32 Index, uint32 SizeX, uint32
 	
 	// Grab the presentation texture out of the swapchain.
 	OutTargetableTexture = OutShaderResourceTexture = (FTexture2DRHIRef&)Swapchain->GetTextureRef();
-
+#if 0
 	// Allocate the depth buffer swapchain while we're here.
 	DepthSwapchain = RenderBridge->CreateSwapchain(Session, PF_DepthStencil, SizeX, SizeY, NumMips, NumSamples, 0, TexCreate_DepthStencilTargetable);
 	if (!DepthSwapchain)
 	{
 		return false;
 	}
+#endif
 	return true;
 }
 
+#if 0
 bool FOpenXRHMD::AllocateDepthTexture(uint32 Index, uint32 SizeX, uint32 SizeY, uint8 Format, uint32 NumMips, uint32 InTexFlags, uint32 TargetableTextureFlags, FTexture2DRHIRef& OutTargetableTexture, FTexture2DRHIRef& OutShaderResourceTexture, uint32 NumSamples)
 {
 	if (!DepthSwapchain || DepthSwapchain->GetTexture2D()->GetSizeX() != SizeX || DepthSwapchain->GetTexture2D()->GetSizeY() != SizeY)
@@ -941,6 +954,7 @@ bool FOpenXRHMD::AllocateDepthTexture(uint32 Index, uint32 SizeX, uint32 SizeY,
 
 	return true;
 }
+#endif
 
 void FOpenXRHMD::OnBeginRendering_RenderThread(FRHICommandListImmediate& RHICmdList, FSceneViewFamily& ViewFamily)
 {
@@ -950,14 +964,19 @@ void FOpenXRHMD::OnBeginRendering_RenderThread(FRHICommandListImmediate& RHICmdL
 		BeginInfo.type = XR_TYPE_FRAME_BEGIN_INFO;
 		BeginInfo.next = nullptr;
 		xrBeginFrame(Session, &BeginInfo);
+		bIsRendering = true;
 	}
 
+	FrameStateRHI = FrameState;
+
 	const FSceneView* MainView = ViewFamily.Views[0];
 	check(MainView);
 	BaseTransform = FTransform(MainView->BaseHmdOrientation, MainView->BaseHmdLocation);
 
 	Swapchain->IncrementSwapChainIndex_RHIThread(FrameStateRHI.predictedDisplayPeriod);
+#if 0
 	DepthSwapchain->IncrementSwapChainIndex_RHIThread(FrameStateRHI.predictedDisplayPeriod);
+#endif
 
 	ViewsRHI.SetNum(Views.Num());
 	DepthLayersRHI.SetNum(Views.Num());
@@ -990,6 +1009,7 @@ void FOpenXRHMD::OnBeginRendering_RenderThread(FRHICommandListImmediate& RHICmdL
 			}
 		};
 
+#if 0
 		DepthLayer.type = XR_TYPE_COMPOSITION_LAYER_DEPTH_INFO_KHR;
 		DepthLayer.next = nullptr;
 		DepthLayer.subImage.swapchain = static_cast<FOpenXRSwapchain*>(GetDepthSwapchain())->GetHandle();
@@ -999,12 +1019,12 @@ void FOpenXRHMD::OnBeginRendering_RenderThread(FRHICommandListImmediate& RHICmdL
 		DepthLayer.maxDepth = 0.0f;
 		DepthLayer.nearZ = NearZ;
 		DepthLayer.farZ = FLT_MAX;
+#endif
 
 		OffsetX += Config.recommendedImageRectWidth;
 	}
 
 	// Give the RHI thread its own copy of the frame state and tracking space
-	FrameStateRHI = FrameState;
 	TrackingSpaceRHI = GetTrackingSpace();
 }
 
@@ -1035,7 +1055,7 @@ void FOpenXRHMD::OnBeginRendering_GameThread()
 	// Pass in a clean FrameState (some implementations may fail if this is just uninitialized).
 	FrameState.type = XR_TYPE_FRAME_STATE;
 	FrameState.next = nullptr;
-	FrameState.predictedDisplayPeriod = FrameStateRHI.predictedDisplayTime = 0;
+	FrameState.predictedDisplayPeriod = 0;
 
 	XR_ENSURE(xrWaitFrame(Session, &WaitInfo, &FrameState));
 
@@ -1044,7 +1064,7 @@ void FOpenXRHMD::OnBeginRendering_GameThread()
 	ViewInfo.type = XR_TYPE_VIEW_LOCATE_INFO;
 	ViewInfo.next = nullptr;
 	ViewInfo.viewConfigurationType = XR_VIEW_CONFIGURATION_TYPE_PRIMARY_STEREO;
-	ViewInfo.space = DeviceSpaces[HMDDeviceId];
+	ViewInfo.space = ActionSpaces[HMDDeviceId].Space;
 	ViewInfo.displayTime = FrameState.predictedDisplayTime;
 	XR_ENSURE(xrLocateViews(Session, &ViewInfo, &ViewState, 0, &ViewCount, nullptr));
 	Views.SetNum(ViewCount);
@@ -1069,7 +1089,7 @@ bool FOpenXRHMD::OnStartGameFrame(FWorldContext& WorldContext)
 
 			CurrentSessionState = SessionState.state;
 
-#if 1
+#if 0
 			static const TCHAR* StateTextMap[] =
 			{
 				TEXT("XR_SESSION_STATE_UNKNOWN"),
@@ -1136,16 +1156,23 @@ bool FOpenXRHMD::OnStartGameFrame(FWorldContext& WorldContext)
 					Session = XR_NULL_HANDLE;
 
 					Swapchain = nullptr;
+#if 0
 					DepthSwapchain = nullptr;
+#endif
 				});
 
 				FlushRenderingCommands();
 
+				bIsRendering = false;
+				bIsReady = false;
 				bIsRunning = false;
 				bRunRequested = false;
 
-				// Clear up device spaces
-				DeviceSpaces.Reset();
+				// Clear up action spaces
+				for (auto& ActionSpace : ActionSpaces)
+				{
+					ActionSpace.DestroySpace();
+				}
 			}
 
 			break;
@@ -1176,9 +1203,11 @@ void FOpenXRHMD::FinishRendering()
 	Layer.views = ViewsRHI.GetData();
 
 	Swapchain->ReleaseCurrentImage_RHIThread();
+#if 0
 	DepthSwapchain->ReleaseCurrentImage_RHIThread();
+#endif
 
-	if (bIsReady)
+	if (bIsRendering)
 	{
 		XrFrameEndInfo EndInfo;
 		XrCompositionLayerBaseHeader* Headers[1] = { reinterpret_cast<XrCompositionLayerBaseHeader*>(&Layer) };
@@ -1321,3 +1350,43 @@ void FOpenXRHMD::DrawHiddenAreaMesh_RenderThread(class FRHICommandList& RHICmdLi
 	RHICmdList.DrawIndexedPrimitive(Mesh.IndexBufferRHI, 0, 0, Mesh.NumVertices, 0, Mesh.NumTriangles, 1);
 #endif
 }
+
+//---------------------------------------------------
+// OpenXR Action Space Implementation
+//---------------------------------------------------
+
+FOpenXRHMD::FActionSpace::FActionSpace(XrAction InAction)
+	: Action(InAction)
+	, Space(XR_NULL_HANDLE)
+{
+}
+
+FOpenXRHMD::FActionSpace::~FActionSpace()
+{
+	DestroySpace();
+}
+
+bool FOpenXRHMD::FActionSpace::CreateSpace(XrSession InSession)
+{
+	if (Action == XR_NULL_HANDLE || Space != XR_NULL_HANDLE)
+	{
+		return false;
+	}
+
+	XrActionSpaceCreateInfo ActionSpaceInfo;
+	ActionSpaceInfo.type = XR_TYPE_ACTION_SPACE_CREATE_INFO;
+	ActionSpaceInfo.next = nullptr;
+	ActionSpaceInfo.subactionPath = XR_NULL_PATH;
+	ActionSpaceInfo.poseInActionSpace = ToXrPose(FTransform::Identity);
+	ActionSpaceInfo.action = Action;
+	return XR_ENSURE(xrCreateActionSpace(InSession, &ActionSpaceInfo, &Space));
+}
+
+void FOpenXRHMD::FActionSpace::DestroySpace()
+{
+	if (Space)
+	{
+		xrDestroySpace(Space);
+	}
+	Space = XR_NULL_HANDLE;
+}
diff --git a/Engine/Plugins/Runtime/OpenXR/Source/OpenXRHMD/Private/OpenXRHMD.h b/Engine/Plugins/Runtime/OpenXR/Source/OpenXRHMD/Private/OpenXRHMD.h
--- a/Engine/Plugins/Runtime/OpenXR/Source/OpenXRHMD/Private/OpenXRHMD.h
+++ b/Engine/Plugins/Runtime/OpenXR/Source/OpenXRHMD/Private/OpenXRHMD.h
@@ -26,6 +26,19 @@ class FOpenXRRenderBridge;
 class FOpenXRHMD : public FHeadMountedDisplayBase, public FXRRenderTargetManager, public FSceneViewExtensionBase
 {
 public:
+	class FActionSpace
+	{
+	public:
+		FActionSpace(XrAction InAction);
+		~FActionSpace();
+
+		bool CreateSpace(XrSession InSession);
+		void DestroySpace();
+
+		XrAction Action;
+		XrSpace Space;
+	};
+
 	class FVulkanExtensions : public IHeadMountedDisplayVulkanExtensions
 	{
 	public:
@@ -136,8 +149,10 @@ public:
 	/** IStereoRenderTargetManager */
 	virtual bool ShouldUseSeparateRenderTarget() const override { return true; }
 	virtual bool AllocateRenderTargetTexture(uint32 Index, uint32 SizeX, uint32 SizeY, uint8 Format, uint32 NumMips, uint32 Flags, uint32 TargetableTextureFlags, FTexture2DRHIRef& OutTargetableTexture, FTexture2DRHIRef& OutShaderResourceTexture, uint32 NumSamples = 1) override;
+#if 0
 	virtual bool NeedReAllocateDepthTexture(const TRefCountPtr<IPooledRenderTarget>& DepthTarget) override final { return false; }
 	virtual bool AllocateDepthTexture(uint32 Index, uint32 SizeX, uint32 SizeY, uint8 Format, uint32 NumMips, uint32 InTexFlags, uint32 TargetableTextureFlags, FTexture2DRHIRef& OutTargetableTexture, FTexture2DRHIRef& OutShaderResourceTexture, uint32 NumSamples = 1) override final;
+#endif
 
 	virtual FXRRenderBridge* GetActiveRenderBridge_GameThread(bool bUseSeparateRenderTarget) override;
 
@@ -161,7 +176,9 @@ public:
 	OPENXRHMD_API int32 AddActionDevice(XrAction Action);
 
 	FXRSwapChain* GetSwapchain() { return Swapchain.Get(); }
+#if 0
 	FXRSwapChain* GetDepthSwapchain() { return DepthSwapchain.Get(); }
+#endif
 	XrInstance GetInstance() { return Instance; }
 	XrSystemId GetSystem() { return System; }
 	XrSession GetSession() { return Session; }
@@ -171,8 +188,10 @@ public:
 	}
 
 private:
+	bool					bStereoEnabled;
 	bool					bIsRunning;
 	bool					bIsReady;
+	bool					bIsRendering;
 	bool					bRunRequested;
 	
 	XrSessionState			CurrentSessionState;
@@ -181,7 +200,6 @@ private:
 	XrInstance				Instance;
 	XrSystemId				System;
 	XrSession				Session;
-	TArray<XrSpace>			DeviceSpaces;
 	XrSpace					LocalSpace;
 	XrSpace					StageSpace;
 	XrSpace					TrackingSpaceRHI;
@@ -199,19 +217,13 @@ private:
 #if 0
 	TArray<FHMDViewMesh>	HiddenAreaMeshes;
 #endif
+	TArray<FActionSpace>	ActionSpaces;
 
 	TRefCountPtr<FOpenXRRenderBridge> RenderBridge;
 	IRendererModule*		RendererModule;
 
 	FXRSwapChainPtr			Swapchain;
+#if 0
 	FXRSwapChainPtr			DepthSwapchain;
-
-	enum KnownDeviceSpaces
-	{
-		DeviceSpaceHMD = 0,
-		// @todo: If we add more known (fixed) spaces, add an enum here.
-		KnownDeviceSpaceCount,
-	};
-
-	TArray<XrAction>		DeferredActionDevices;
+#endif
 };
diff --git a/Engine/Plugins/Runtime/OpenXR/Source/OpenXRHMD/Private/OpenXRHMD_Swapchain.cpp b/Engine/Plugins/Runtime/OpenXR/Source/OpenXRHMD/Private/OpenXRHMD_Swapchain.cpp
--- a/Engine/Plugins/Runtime/OpenXR/Source/OpenXRHMD/Private/OpenXRHMD_Swapchain.cpp
+++ b/Engine/Plugins/Runtime/OpenXR/Source/OpenXRHMD/Private/OpenXRHMD_Swapchain.cpp
@@ -57,8 +57,13 @@ void FOpenXRSwapchain::ReleaseResources_RHIThread()
 	xrDestroySwapchain(Handle);
 }
 
-XrSwapchain CreateSwapchain(XrSession InSession, uint32 PlatformFormat, uint32 SizeX, uint32 SizeY, uint32 NumMips, uint32 NumSamples, uint32 TargetableTextureFlags)
+uint8 GetNearestSupportedSwapchainFormat(XrSession InSession, uint8 RequestedFormat, TFunction<uint32(uint8)> ToPlatformFormat = nullptr)
 {
+	if (!ToPlatformFormat)
+	{
+		ToPlatformFormat = [](uint8 InFormat) { return GPixelFormats[InFormat].PlatformFormat; };
+	}
+
 	uint32_t FormatsCount;
 	xrEnumerateSwapchainFormats(InSession, 0, &FormatsCount, nullptr);
 
@@ -67,12 +72,44 @@ XrSwapchain CreateSwapchain(XrSession InSession, uint32 PlatformFormat, uint32 S
 	XR_ENSURE(xrEnumerateSwapchainFormats(InSession, (uint32_t)Formats.Num(), &FormatsCount, Formats.GetData()));
 	ensure(FormatsCount == Formats.Num());
 
-	if (!Formats.Contains(PlatformFormat))
+	// Return immediately if the runtime supports the exact format being requested.
+	uint32 PlatformFormat = ToPlatformFormat(RequestedFormat);
+	if (Formats.Contains(PlatformFormat))
 	{
-		UE_LOG(LogHMD, Log, TEXT("Swapchain format not supported (%d), falling back to runtime preferred format (%d)."), PlatformFormat, Formats[0]);
-		PlatformFormat = Formats[0];
+		return RequestedFormat;
 	}
 
+	// Search for an 8bpc fallback format in order of preference (first element in the array has the high preference).
+	uint8 FallbackFormat = 0;
+	uint32 FallbackPlatformFormat = 0;
+	for (int64_t Format : Formats)
+	{
+		if (Format == ToPlatformFormat(PF_B8G8R8A8))
+		{
+			FallbackFormat = PF_B8G8R8A8;
+			FallbackPlatformFormat = Format;
+			break;
+		}
+		else if (Format == ToPlatformFormat(PF_R8G8B8A8))
+		{
+			FallbackFormat = PF_R8G8B8A8;
+			FallbackPlatformFormat = Format;
+			break;
+		}
+	}
+
+	if (!FallbackFormat)
+	{
+		UE_LOG(LogHMD, Warning, TEXT("No compatible swapchain format found!"));
+		return PF_Unknown;
+	}
+
+	UE_LOG(LogHMD, Warning, TEXT("Swapchain format not supported (%d), falling back to runtime preferred format (%d)."), PlatformFormat, FallbackPlatformFormat);
+	return FallbackFormat;
+}
+
+XrSwapchain CreateSwapchain(XrSession InSession, uint32 PlatformFormat, uint32 SizeX, uint32 SizeY, uint32 NumMips, uint32 NumSamples, uint32 TargetableTextureFlags)
+{
 	XrSwapchainUsageFlags Usage = 0;
 	if (TargetableTextureFlags & TexCreate_RenderTargetable)
 	{
@@ -129,11 +166,22 @@ TArray<T> EnumerateImages(XrSwapchain InSwapchain, XrStructureType InType)
 #ifdef XR_USE_GRAPHICS_API_D3D11
 FXRSwapChainPtr CreateSwapchain_D3D11(XrSession InSession, uint8 Format, uint32 SizeX, uint32 SizeY, uint32 NumMips, uint32 NumSamples, uint32 Flags, uint32 TargetableTextureFlags)
 {
-	// We need to convert typeless to typed formats to create a swapchain
-	DXGI_FORMAT PlatformFormat = (DXGI_FORMAT)GPixelFormats[Format].PlatformFormat;
-	PlatformFormat = FindDepthStencilDXGIFormat(PlatformFormat);
-	PlatformFormat = FindShaderResourceDXGIFormat(PlatformFormat, Flags & TexCreate_SRGB);
-	XrSwapchain Swapchain = CreateSwapchain(InSession, PlatformFormat, SizeX, SizeY, NumMips, NumSamples, TargetableTextureFlags);
+	TFunction<uint32(uint8)> ToPlatformFormat = [Flags](uint8 InFormat)
+	{
+		// We need to convert typeless to typed formats to create a swapchain
+		DXGI_FORMAT PlatformFormat = (DXGI_FORMAT)GPixelFormats[InFormat].PlatformFormat;
+		PlatformFormat = FindDepthStencilDXGIFormat(PlatformFormat);
+		PlatformFormat = FindShaderResourceDXGIFormat(PlatformFormat, Flags & TexCreate_SRGB);
+		return PlatformFormat;
+	};
+
+	Format = GetNearestSupportedSwapchainFormat(InSession, Format, ToPlatformFormat);
+	if (!Format)
+	{
+		return nullptr;
+	}
+
+	XrSwapchain Swapchain = CreateSwapchain(InSession, ToPlatformFormat(Format), SizeX, SizeY, NumMips, NumSamples, TargetableTextureFlags);
 	if (!Swapchain)
 	{
 		return nullptr;
@@ -156,11 +204,22 @@ FXRSwapChainPtr CreateSwapchain_D3D11(XrSession InSession, uint8 Format, uint32
 #ifdef XR_USE_GRAPHICS_API_D3D12
 FXRSwapChainPtr CreateSwapchain_D3D12(XrSession InSession, uint8 Format, uint32 SizeX, uint32 SizeY, uint32 NumMips, uint32 NumSamples, uint32 Flags, uint32 TargetableTextureFlags)
 {
-	// We need to convert typeless to typed formats to create a swapchain
-	DXGI_FORMAT PlatformFormat = (DXGI_FORMAT)GPixelFormats[Format].PlatformFormat;
-	PlatformFormat = FindDepthStencilDXGIFormat_D3D12(PlatformFormat);
-	PlatformFormat = FindShaderResourceDXGIFormat_D3D12(PlatformFormat, Flags & TexCreate_SRGB);
-	XrSwapchain Swapchain = CreateSwapchain(InSession, Format, SizeX, SizeY, NumMips, NumSamples, TargetableTextureFlags);
+	TFunction<uint32(uint8)> ToPlatformFormat = [Flags](uint8 InFormat)
+	{
+		// We need to convert typeless to typed formats to create a swapchain
+		DXGI_FORMAT PlatformFormat = (DXGI_FORMAT)GPixelFormats[InFormat].PlatformFormat;
+		PlatformFormat = FindDepthStencilDXGIFormat_D3D12(PlatformFormat);
+		PlatformFormat = FindShaderResourceDXGIFormat_D3D12(PlatformFormat, Flags & TexCreate_SRGB);
+		return PlatformFormat;
+	};
+
+	Format = GetNearestSupportedSwapchainFormat(InSession, Format, ToPlatformFormat);
+	if (!Format)
+	{
+		return nullptr;
+	}
+
+	XrSwapchain Swapchain = CreateSwapchain(InSession, ToPlatformFormat(Format), SizeX, SizeY, NumMips, NumSamples, TargetableTextureFlags);
 	if (!Swapchain)
 	{
 		return nullptr;
@@ -183,6 +242,12 @@ FXRSwapChainPtr CreateSwapchain_D3D12(XrSession InSession, uint8 Format, uint32
 #ifdef XR_USE_GRAPHICS_API_OPENGL
 FXRSwapChainPtr CreateSwapchain_OpenGL(XrSession InSession, uint8 Format, uint32 SizeX, uint32 SizeY, uint32 NumMips, uint32 NumSamples, uint32 Flags, uint32 TargetableTextureFlags)
 {
+	Format = GetNearestSupportedSwapchainFormat(InSession, Format);
+	if (!Format)
+	{
+		return nullptr;
+	}
+
 	XrSwapchain Swapchain = CreateSwapchain(InSession, GPixelFormats[Format].PlatformFormat, SizeX, SizeY, NumMips, NumSamples, TargetableTextureFlags);
 	if (!Swapchain)
 	{
@@ -206,6 +271,12 @@ FXRSwapChainPtr CreateSwapchain_OpenGL(XrSession InSession, uint8 Format, uint32
 #ifdef XR_USE_GRAPHICS_API_VULKAN
 FXRSwapChainPtr CreateSwapchain_Vulkan(XrSession InSession, uint8 Format, uint32 SizeX, uint32 SizeY, uint32 NumMips, uint32 NumSamples, uint32 Flags, uint32 TargetableTextureFlags)
 {
+	Format = GetNearestSupportedSwapchainFormat(InSession, Format);
+	if (!Format)
+	{
+		return nullptr;
+	}
+
 	XrSwapchain Swapchain = CreateSwapchain(InSession, GPixelFormats[Format].PlatformFormat, SizeX, SizeY, NumMips, NumSamples, TargetableTextureFlags);
 	if (!Swapchain)
 	{
diff --git a/Engine/Plugins/Runtime/OpenXR/Source/OpenXRHMD/Public/IOpenXRHMDPlugin.h b/Engine/Plugins/Runtime/OpenXR/Source/OpenXRHMD/Public/IOpenXRHMDPlugin.h
--- a/Engine/Plugins/Runtime/OpenXR/Source/OpenXRHMD/Public/IOpenXRHMDPlugin.h
+++ b/Engine/Plugins/Runtime/OpenXR/Source/OpenXRHMD/Public/IOpenXRHMDPlugin.h
@@ -1,4 +1,4 @@
-// Copyright 1998-2018 Epic Games, Inc. All Rights Reserved.
+// Copyright 1998-2019 Epic Games, Inc. All Rights Reserved.
 
 #pragma once
 
diff --git a/Engine/Plugins/Runtime/OpenXR/Source/OpenXRInput/Private/OpenXRInput.cpp b/Engine/Plugins/Runtime/OpenXR/Source/OpenXRInput/Private/OpenXRInput.cpp
--- a/Engine/Plugins/Runtime/OpenXR/Source/OpenXRInput/Private/OpenXRInput.cpp
+++ b/Engine/Plugins/Runtime/OpenXR/Source/OpenXRInput/Private/OpenXRInput.cpp
@@ -107,7 +107,7 @@ FOpenXRInputPlugin::FOpenXRAction::FOpenXRAction(XrActionSet InSet, const FInput
 }
 
 FOpenXRInputPlugin::FOpenXRAction::FOpenXRAction(XrActionSet InSet, const FInputAxisKeyMapping& InAxisKey)
-	: FOpenXRAction(InSet, XR_ACTION_TYPE_VECTOR2F_INPUT, InAxisKey.AxisName)
+	: FOpenXRAction(InSet, XR_ACTION_TYPE_FLOAT_INPUT, InAxisKey.AxisName)
 {
 	ActionKey = InAxisKey.Key.GetFName();
 }
@@ -168,8 +168,8 @@ FOpenXRInputPlugin::FOpenXRInput::FOpenXRInput(FOpenXRHMD* HMD)
 	Controllers.Add(EControllerHand::Left, FOpenXRController(OpenXRHMD, ActionSet, "Left Controller"));
 	Controllers.Add(EControllerHand::Right, FOpenXRController(OpenXRHMD, ActionSet, "Right Controller"));
 
-	Bindings.Add(XrActionSuggestedBinding{ Controllers[EControllerHand::Left].Action, GetPath(Instance, "/user/hand/left/input/palm") });
-	Bindings.Add(XrActionSuggestedBinding{ Controllers[EControllerHand::Right].Action, GetPath(Instance, "/user/hand/right/input/palm") });
+	Bindings.Add(XrActionSuggestedBinding{ Controllers[EControllerHand::Left].Action, GetPath(Instance, "/user/hand/left/input/grip/pose") });
+	Bindings.Add(XrActionSuggestedBinding{ Controllers[EControllerHand::Right].Action, GetPath(Instance, "/user/hand/right/input/grip/pose") });
 
 	Bindings.Add(XrActionSuggestedBinding{ Controllers[EControllerHand::Left].VibrationAction, GetPath(Instance, "/user/hand/left/output/haptic") });
 	Bindings.Add(XrActionSuggestedBinding{ Controllers[EControllerHand::Right].VibrationAction, GetPath(Instance, "/user/hand/right/output/haptic") });
@@ -177,8 +177,8 @@ FOpenXRInputPlugin::FOpenXRInput::FOpenXRInput(FOpenXRHMD* HMD)
 	InteractionMappings.Add(FGamepadKeyNames::MotionController_Left_Shoulder, GetPath(Instance, "/user/hand/left/input/menu/click"));
 	InteractionMappings.Add(FGamepadKeyNames::MotionController_Left_Trigger, GetPath(Instance, "/user/hand/left/input/trigger"));
 	InteractionMappings.Add(FGamepadKeyNames::MotionController_Left_TriggerAxis, GetPath(Instance, "/user/hand/left/input/trigger"));
-	InteractionMappings.Add(FGamepadKeyNames::MotionController_Left_Grip1, GetPath(Instance, "/user/hand/left/input/grip"));
-	InteractionMappings.Add(FGamepadKeyNames::MotionController_Left_Grip1Axis, GetPath(Instance, "/user/hand/left/input/grip"));
+	InteractionMappings.Add(FGamepadKeyNames::MotionController_Left_Grip1, GetPath(Instance, "/user/hand/left/input/squeeze"));
+	InteractionMappings.Add(FGamepadKeyNames::MotionController_Left_Grip1Axis, GetPath(Instance, "/user/hand/left/input/squeeze"));
 	InteractionMappings.Add(FGamepadKeyNames::MotionController_Left_Thumbstick_X, GetPath(Instance, "/user/hand/left/input/thumbstick/x"));
 	InteractionMappings.Add(FGamepadKeyNames::MotionController_Left_Thumbstick_Y, GetPath(Instance, "/user/hand/left/input/thumbstick/y"));
 	InteractionMappings.Add(FGamepadKeyNames::MotionController_Left_Thumbstick, GetPath(Instance, "/user/hand/left/input/thumbstick/click"));
@@ -188,8 +188,8 @@ FOpenXRInputPlugin::FOpenXRInput::FOpenXRInput(FOpenXRHMD* HMD)
 	InteractionMappings.Add(FGamepadKeyNames::MotionController_Right_Shoulder, GetPath(Instance, "/user/hand/right/input/menu/click"));
 	InteractionMappings.Add(FGamepadKeyNames::MotionController_Right_Trigger, GetPath(Instance, "/user/hand/right/input/trigger"));
 	InteractionMappings.Add(FGamepadKeyNames::MotionController_Right_TriggerAxis, GetPath(Instance, "/user/hand/right/input/trigger"));
-	InteractionMappings.Add(FGamepadKeyNames::MotionController_Right_Grip1, GetPath(Instance, "/user/hand/right/input/grip"));
-	InteractionMappings.Add(FGamepadKeyNames::MotionController_Right_Grip1Axis, GetPath(Instance, "/user/hand/right/input/grip"));
+	InteractionMappings.Add(FGamepadKeyNames::MotionController_Right_Grip1, GetPath(Instance, "/user/hand/right/input/squeeze"));
+	InteractionMappings.Add(FGamepadKeyNames::MotionController_Right_Grip1Axis, GetPath(Instance, "/user/hand/right/input/squeeze"));
 	InteractionMappings.Add(FGamepadKeyNames::MotionController_Right_Thumbstick_X, GetPath(Instance, "/user/hand/right/input/thumbstick/x"));
 	InteractionMappings.Add(FGamepadKeyNames::MotionController_Right_Thumbstick_Y, GetPath(Instance, "/user/hand/right/input/thumbstick/y"));
 	InteractionMappings.Add(FGamepadKeyNames::MotionController_Right_Thumbstick, GetPath(Instance, "/user/hand/right/input/thumbstick/click"));
