Fixes several issues (submitting work done by [REDACTED_PERSON] after reviewing):
* Fix swapchain lifetime issues (don't set swapchain to nullptr in middle of a frame)
* added missing check for null space
* Fix build server warnings on error logging (XR_ENSURE) macro with new lambda-based macro (which also provides correct __FILE__ and __LINE__ outputs which the inline function did not).

diff --git a/Engine/Plugins/Runtime/OpenXR/Source/OpenXRHMD/Private/OpenXRHMD.cpp b/Engine/Plugins/Runtime/OpenXR/Source/OpenXRHMD/Private/OpenXRHMD.cpp
--- a/Engine/Plugins/Runtime/OpenXR/Source/OpenXRHMD/Private/OpenXRHMD.cpp
+++ b/Engine/Plugins/Runtime/OpenXR/Source/OpenXRHMD/Private/OpenXRHMD.cpp
@@ -477,7 +477,7 @@ bool FOpenXRHMD::GetCurrentPose(int32 DeviceId, FQuat& CurrentOrientation, FVect
 	CurrentOrientation = FQuat::Identity;
 	CurrentPosition = FVector::ZeroVector;
 
-	if (!ActionSpaces.IsValidIndex(DeviceId) || FrameState.predictedDisplayTime <= 0)
+	if (!ActionSpaces.IsValidIndex(DeviceId) || !ActionSpaces[DeviceId].Space || FrameState.predictedDisplayTime <= 0)
 	{
 		return false;
 	}
@@ -965,6 +965,11 @@ void FOpenXRHMD::OnBeginRendering_RenderThread(FRHICommandListImmediate& RHICmdL
 		BeginInfo.next = nullptr;
 		xrBeginFrame(Session, &BeginInfo);
 		bIsRendering = true;
+
+		Swapchain->IncrementSwapChainIndex_RHIThread(FrameStateRHI.predictedDisplayPeriod);
+#if 0
+		DepthSwapchain->IncrementSwapChainIndex_RHIThread(FrameStateRHI.predictedDisplayPeriod);
+#endif
 	}
 
 	FrameStateRHI = FrameState;
@@ -973,11 +978,6 @@ void FOpenXRHMD::OnBeginRendering_RenderThread(FRHICommandListImmediate& RHICmdL
 	check(MainView);
 	BaseTransform = FTransform(MainView->BaseHmdOrientation, MainView->BaseHmdLocation);
 
-	Swapchain->IncrementSwapChainIndex_RHIThread(FrameStateRHI.predictedDisplayPeriod);
-#if 0
-	DepthSwapchain->IncrementSwapChainIndex_RHIThread(FrameStateRHI.predictedDisplayPeriod);
-#endif
-
 	ViewsRHI.SetNum(Views.Num());
 	DepthLayersRHI.SetNum(Views.Num());
 
@@ -1154,11 +1154,6 @@ bool FOpenXRHMD::OnStartGameFrame(FWorldContext& WorldContext)
 					XR_ENSURE(xrDestroySession(Session));
 
 					Session = XR_NULL_HANDLE;
-
-					Swapchain = nullptr;
-#if 0
-					DepthSwapchain = nullptr;
-#endif
 				});
 
 				FlushRenderingCommands();
@@ -1202,13 +1197,13 @@ void FOpenXRHMD::FinishRendering()
 	Layer.viewCount = ViewsRHI.Num();
 	Layer.views = ViewsRHI.GetData();
 
-	Swapchain->ReleaseCurrentImage_RHIThread();
+	if (bIsRendering)
+	{
+		Swapchain->ReleaseCurrentImage_RHIThread();
 #if 0
-	DepthSwapchain->ReleaseCurrentImage_RHIThread();
+		DepthSwapchain->ReleaseCurrentImage_RHIThread();
 #endif
 
-	if (bIsRendering)
-	{
 		XrFrameEndInfo EndInfo;
 		XrCompositionLayerBaseHeader* Headers[1] = { reinterpret_cast<XrCompositionLayerBaseHeader*>(&Layer) };
 		EndInfo.type = XR_TYPE_FRAME_END_INFO;
diff --git a/Engine/Plugins/Runtime/OpenXR/Source/OpenXRHMD/Private/OpenXRHMDPrivate.h b/Engine/Plugins/Runtime/OpenXR/Source/OpenXRHMD/Private/OpenXRHMDPrivate.h
--- a/Engine/Plugins/Runtime/OpenXR/Source/OpenXRHMD/Private/OpenXRHMDPrivate.h
+++ b/Engine/Plugins/Runtime/OpenXR/Source/OpenXRHMD/Private/OpenXRHMDPrivate.h
@@ -24,12 +24,14 @@ constexpr const TCHAR* xrResultToString(XrResult e)
 	}
 }
 
-inline bool LogXrResult(XrResult Result)
-{
-	return ensureMsgf(XR_SUCCEEDED(Result), TEXT("OpenXR call failed with result %s"), xrResultToString(Result));
-}
-
-#define XR_ENSURE(x) LogXrResult(x)
+#if DO_CHECK
+#define XR_ENSURE(x) [] (XrResult Result) \
+	{ \
+		return ensureMsgf(XR_SUCCEEDED(Result), TEXT("OpenXR call failed with result %s"), xrResultToString(Result)); \
+	} (x)
+#else
+#define XR_ENSURE(x) XR_SUCCEEDED(x)
+#endif
 
 FORCEINLINE FQuat ToFQuat(XrQuaternionf Quat)
 {
