Fix static analysis warnings

diff --git a/Engine/Plugins/Runtime/OpenXR/Source/OpenXRHMD/Private/OpenXRHMD.cpp b/Engine/Plugins/Runtime/OpenXR/Source/OpenXRHMD/Private/OpenXRHMD.cpp
--- a/Engine/Plugins/Runtime/OpenXR/Source/OpenXRHMD/Private/OpenXRHMD.cpp
+++ b/Engine/Plugins/Runtime/OpenXR/Source/OpenXRHMD/Private/OpenXRHMD.cpp
@@ -795,7 +795,7 @@ bool FOpenXRHMD::OnStereoStartup()
 	XR_ENSURE(xrEnumerateReferenceSpaces(Session, (uint32_t)spaces.Num(), &referenceSpacesCount, spaces.GetData()));
 	ensure(referenceSpacesCount == spaces.Num());
 
-	XrSpace HmdSpace;
+	XrSpace HmdSpace = XR_NULL_HANDLE;
 	XrReferenceSpaceCreateInfo SpaceInfo;
 
 	ensure(spaces.Contains(XR_REFERENCE_SPACE_TYPE_VIEW));
diff --git a/Engine/Plugins/Runtime/OpenXR/Source/OpenXRHMD/Private/OpenXRHMD_RenderBridge.cpp b/Engine/Plugins/Runtime/OpenXR/Source/OpenXRHMD/Private/OpenXRHMD_RenderBridge.cpp
--- a/Engine/Plugins/Runtime/OpenXR/Source/OpenXRHMD/Private/OpenXRHMD_RenderBridge.cpp
+++ b/Engine/Plugins/Runtime/OpenXR/Source/OpenXRHMD/Private/OpenXRHMD_RenderBridge.cpp
@@ -100,6 +100,8 @@ public:
 		XrGraphicsRequirementsOpenGLKHR Requirements;
 		Requirements.type = XR_TYPE_GRAPHICS_REQUIREMENTS_OPENGL_KHR;
 		Requirements.next = nullptr;
+		Requirements.minApiVersionSupported = 0;
+		Requirements.maxApiVersionSupported = 0;
 		XR_ENSURE(xrGetOpenGLGraphicsRequirementsKHR(InInstance, InSystem, &Requirements));
 
 		XrVersion RHIVersion = XR_MAKE_VERSION(FOpenGL::GetMajorVersion(), FOpenGL::GetMinorVersion(), 0);
@@ -152,6 +154,8 @@ public:
 		XrGraphicsRequirementsVulkanKHR Requirements;
 		Requirements.type = XR_TYPE_GRAPHICS_REQUIREMENTS_VULKAN_KHR;
 		Requirements.next = nullptr;
+		Requirements.minApiVersionSupported = 0;
+		Requirements.maxApiVersionSupported = 0;
 		XR_ENSURE(xrGetVulkanGraphicsRequirementsKHR(InInstance, InSystem, &Requirements));
 
 		// The extension uses the OpenXR version format instead of the Vulkan one
diff --git a/Engine/Plugins/Runtime/OpenXR/Source/OpenXRInput/Private/OpenXRInput.cpp b/Engine/Plugins/Runtime/OpenXR/Source/OpenXRInput/Private/OpenXRInput.cpp
--- a/Engine/Plugins/Runtime/OpenXR/Source/OpenXRInput/Private/OpenXRInput.cpp
+++ b/Engine/Plugins/Runtime/OpenXR/Source/OpenXRInput/Private/OpenXRInput.cpp
@@ -155,7 +155,7 @@ FOpenXRInputPlugin::FOpenXRInput::FOpenXRInput(FOpenXRHMD* HMD)
 	XrInstance Instance = OpenXRHMD->GetInstance();
 	check(Instance);
 
-	XrActionSet ActionSet;
+	XrActionSet ActionSet = XR_NULL_HANDLE;
 	XrActionSetCreateInfo SetInfo;
 	SetInfo.type = XR_TYPE_ACTION_SET_CREATE_INFO;
 	SetInfo.next = nullptr;
