Disable unsupported paths for now
We'll enable these again once we refactor input

diff --git a/Engine/Plugins/Runtime/OpenXR/Source/OpenXRInput/Private/OpenXRInput.cpp b/Engine/Plugins/Runtime/OpenXR/Source/OpenXRInput/Private/OpenXRInput.cpp
--- a/Engine/Plugins/Runtime/OpenXR/Source/OpenXRInput/Private/OpenXRInput.cpp
+++ b/Engine/Plugins/Runtime/OpenXR/Source/OpenXRInput/Private/OpenXRInput.cpp
@@ -182,8 +182,8 @@ FOpenXRInputPlugin::FOpenXRInput::FOpenXRInput(FOpenXRHMD* HMD)
 	InteractionMappings.Add(FGamepadKeyNames::MotionController_Left_Thumbstick_X, GetPath(Instance, "/user/hand/left/input/thumbstick/x"));
 	InteractionMappings.Add(FGamepadKeyNames::MotionController_Left_Thumbstick_Y, GetPath(Instance, "/user/hand/left/input/thumbstick/y"));
 	InteractionMappings.Add(FGamepadKeyNames::MotionController_Left_Thumbstick, GetPath(Instance, "/user/hand/left/input/thumbstick/click"));
-	InteractionMappings.Add(FGamepadKeyNames::MotionController_Left_FaceButton1, GetPath(Instance, "/user/hand/left/input/x/click"));
-	InteractionMappings.Add(FGamepadKeyNames::MotionController_Left_FaceButton2, GetPath(Instance, "/user/hand/left/input/y/click"));
+	//InteractionMappings.Add(FGamepadKeyNames::MotionController_Left_FaceButton1, GetPath(Instance, "/user/hand/left/input/x/click"));
+	//InteractionMappings.Add(FGamepadKeyNames::MotionController_Left_FaceButton2, GetPath(Instance, "/user/hand/left/input/y/click"));
 
 	InteractionMappings.Add(FGamepadKeyNames::MotionController_Right_Shoulder, GetPath(Instance, "/user/hand/right/input/menu/click"));
 	InteractionMappings.Add(FGamepadKeyNames::MotionController_Right_Trigger, GetPath(Instance, "/user/hand/right/input/trigger"));
@@ -193,8 +193,8 @@ FOpenXRInputPlugin::FOpenXRInput::FOpenXRInput(FOpenXRHMD* HMD)
 	InteractionMappings.Add(FGamepadKeyNames::MotionController_Right_Thumbstick_X, GetPath(Instance, "/user/hand/right/input/thumbstick/x"));
 	InteractionMappings.Add(FGamepadKeyNames::MotionController_Right_Thumbstick_Y, GetPath(Instance, "/user/hand/right/input/thumbstick/y"));
 	InteractionMappings.Add(FGamepadKeyNames::MotionController_Right_Thumbstick, GetPath(Instance, "/user/hand/right/input/thumbstick/click"));
-	InteractionMappings.Add(FGamepadKeyNames::MotionController_Right_FaceButton1, GetPath(Instance, "/user/hand/right/input/a/click"));
-	InteractionMappings.Add(FGamepadKeyNames::MotionController_Right_FaceButton2, GetPath(Instance, "/user/hand/right/input/b/click"));
+	//InteractionMappings.Add(FGamepadKeyNames::MotionController_Right_FaceButton1, GetPath(Instance, "/user/hand/right/input/a/click"));
+	//InteractionMappings.Add(FGamepadKeyNames::MotionController_Right_FaceButton2, GetPath(Instance, "/user/hand/right/input/b/click"));
 
 	auto InputSettings = GetDefault<UInputSettings>();
 	if (InputSettings != nullptr)
@@ -228,10 +228,10 @@ FOpenXRInputPlugin::FOpenXRInput::FOpenXRInput(FOpenXRHMD* HMD)
 	}
 
 	TArray<XrPath> Profiles;
-	Profiles.Add(GetPath(Instance, "/interaction_profiles/khr/simple_controller"));
+	//Profiles.Add(GetPath(Instance, "/interaction_profiles/khr/simple_controller"));
 	Profiles.Add(GetPath(Instance, "/interaction_profiles/microsoft/motion_controller"));
 	Profiles.Add(GetPath(Instance, "/interaction_profiles/oculus/touch_controller"));
-	Profiles.Add(GetPath(Instance, "/interaction_profiles/valve/knuckles_controller"));
+	Profiles.Add(GetPath(Instance, "/interaction_profiles/valve/index_controller"));
 
 	for (XrPath Profile : Profiles)
 	{
