Suppress PVS Studio false positives in OpenXR plugin

diff --git a/Engine/Plugins/Runtime/OpenXR/Source/OpenXRHMD/Private/OpenXRHMD_RenderBridge.cpp b/Engine/Plugins/Runtime/OpenXR/Source/OpenXRHMD/Private/OpenXRHMD_RenderBridge.cpp
--- a/Engine/Plugins/Runtime/OpenXR/Source/OpenXRHMD/Private/OpenXRHMD_RenderBridge.cpp
+++ b/Engine/Plugins/Runtime/OpenXR/Source/OpenXRHMD/Private/OpenXRHMD_RenderBridge.cpp
@@ -105,12 +105,12 @@ public:
 		XR_ENSURE(xrGetOpenGLGraphicsRequirementsKHR(InInstance, InSystem, &Requirements));
 
 		XrVersion RHIVersion = XR_MAKE_VERSION(FOpenGL::GetMajorVersion(), FOpenGL::GetMinorVersion(), 0);
-		if (RHIVersion < Requirements.minApiVersionSupported)
+		if (RHIVersion < Requirements.minApiVersionSupported) //-V547
 		{
 			UE_LOG(LogHMD, Fatal, TEXT("The OpenGL API version does not meet the minimum version required by the OpenXR runtime"));
 		}
 
-		if (RHIVersion > Requirements.maxApiVersionSupported)
+		if (RHIVersion > Requirements.maxApiVersionSupported) //-V547
 		{
 			UE_LOG(LogHMD, Warning, TEXT("The OpenGL API version has not been tested with the OpenXR runtime"));
 		}
@@ -164,12 +164,12 @@ public:
 			VK_VERSION_MINOR(UE_VK_API_VERSION),
 			VK_VERSION_PATCH(UE_VK_API_VERSION)
 		);
-		if (RHIVersion < Requirements.minApiVersionSupported)
+		if (RHIVersion < Requirements.minApiVersionSupported) //-V547
 		{
 			UE_LOG(LogHMD, Fatal, TEXT("The Vulkan API version does not meet the minimum version required by the OpenXR runtime"));
 		}
 
-		if (RHIVersion > Requirements.maxApiVersionSupported)
+		if (RHIVersion > Requirements.maxApiVersionSupported) //-V547
 		{
 			UE_LOG(LogHMD, Warning, TEXT("The Vulkan API version has not been tested with the OpenXR runtime"));
 		}
