Fixing OpenXR validation errors

diff --git a/Engine/Plugins/Runtime/OpenXR/Source/OpenXRHMD/Private/OpenXRHMD.cpp b/Engine/Plugins/Runtime/OpenXR/Source/OpenXRHMD/Private/OpenXRHMD.cpp
--- a/Engine/Plugins/Runtime/OpenXR/Source/OpenXRHMD/Private/OpenXRHMD.cpp
+++ b/Engine/Plugins/Runtime/OpenXR/Source/OpenXRHMD/Private/OpenXRHMD.cpp
@@ -880,7 +880,7 @@ bool FOpenXRHMD::IsRunning() const
 	return bIsRunning;
 }
 
-bool FOpenXRHMD::IsRendering() const
+bool FOpenXRHMD::IsFocused() const
 {
 	return CurrentSessionState == XR_SESSION_STATE_FOCUSED;
 }
@@ -1071,14 +1071,21 @@ void FOpenXRHMD::OnBeginRendering_GameThread()
 	XR_ENSURE(xrLocateViews(Session, &ViewInfo, &ViewState, Views.Num(), &ViewCount, Views.GetData()));
 }
 
+bool FOpenXRHMD::ReadNextEvent(XrEventDataBuffer* buffer)
+{
+	// It is sufficient to clear just the XrEventDataBuffer header to XR_TYPE_EVENT_DATA_BUFFER
+	XrEventDataBaseHeader* baseHeader = reinterpret_cast<XrEventDataBaseHeader*>(buffer);
+	*baseHeader = { XR_TYPE_EVENT_DATA_BUFFER };
+	const XrResult xr = xrPollEvent(Instance, buffer);
+	XR_ENSURE(xr);
+	return xr == XR_SUCCESS;
+}
+
 bool FOpenXRHMD::OnStartGameFrame(FWorldContext& WorldContext)
 {
-	// Initialize an event buffer to hold the output.
+	// Process all pending messages.
 	XrEventDataBuffer event;
-	// Only the header needs to be initialized.
-	event.type = XR_TYPE_EVENT_DATA_BUFFER;
-	event.next = nullptr;
-	while (xrPollEvent(Instance, &event) == XR_SUCCESS) 
+	while (ReadNextEvent(&event))
 	{
 		switch (event.type) 
 		{
@@ -1143,6 +1150,12 @@ bool FOpenXRHMD::OnStartGameFrame(FWorldContext& WorldContext)
 
 			if (Session != XR_NULL_HANDLE)
 			{
+				// Clear up action spaces
+				for (auto& ActionSpace : ActionSpaces)
+				{
+					ActionSpace.DestroySpace();
+				}
+
 				// Close the session now we're allowed to.
 				ENQUEUE_RENDER_COMMAND(OpenXRDestroySession)([this](FRHICommandListImmediate& RHICmdList)
 				{
@@ -1162,12 +1175,6 @@ bool FOpenXRHMD::OnStartGameFrame(FWorldContext& WorldContext)
 				bIsReady = false;
 				bIsRunning = false;
 				bRunRequested = false;
-
-				// Clear up action spaces
-				for (auto& ActionSpace : ActionSpaces)
-				{
-					ActionSpace.DestroySpace();
-				}
 			}
 
 			break;
@@ -1356,11 +1363,6 @@ FOpenXRHMD::FActionSpace::FActionSpace(XrAction InAction)
 {
 }
 
-FOpenXRHMD::FActionSpace::~FActionSpace()
-{
-	DestroySpace();
-}
-
 bool FOpenXRHMD::FActionSpace::CreateSpace(XrSession InSession)
 {
 	if (Action == XR_NULL_HANDLE || Space != XR_NULL_HANDLE)
@@ -1381,7 +1383,7 @@ void FOpenXRHMD::FActionSpace::DestroySpace()
 {
 	if (Space)
 	{
-		xrDestroySpace(Space);
+		XR_ENSURE(xrDestroySpace(Space));
 	}
 	Space = XR_NULL_HANDLE;
 }
diff --git a/Engine/Plugins/Runtime/OpenXR/Source/OpenXRHMD/Private/OpenXRHMD.h b/Engine/Plugins/Runtime/OpenXR/Source/OpenXRHMD/Private/OpenXRHMD.h
--- a/Engine/Plugins/Runtime/OpenXR/Source/OpenXRHMD/Private/OpenXRHMD.h
+++ b/Engine/Plugins/Runtime/OpenXR/Source/OpenXRHMD/Private/OpenXRHMD.h
@@ -30,7 +30,6 @@ public:
 	{
 	public:
 		FActionSpace(XrAction InAction);
-		~FActionSpace();
 
 		bool CreateSpace(XrSession InSession);
 		void DestroySpace();
@@ -104,6 +103,7 @@ protected:
 	bool StartSession();
 	bool OnStereoStartup();
 	bool OnStereoTeardown();
+	bool ReadNextEvent(XrEventDataBuffer* buffer);
 
 public:
 	/** IHeadMountedDisplay interface */
@@ -170,7 +170,7 @@ public:
 	/** @return	True if the HMD was initialized OK */
 	OPENXRHMD_API bool IsInitialized() const;
 	OPENXRHMD_API bool IsRunning() const;
-	OPENXRHMD_API bool IsRendering() const;
+	OPENXRHMD_API bool IsFocused() const;
 	void FinishRendering();
 
 	OPENXRHMD_API int32 AddActionDevice(XrAction Action);
diff --git a/Engine/Plugins/Runtime/OpenXR/Source/OpenXRInput/Private/OpenXRInput.cpp b/Engine/Plugins/Runtime/OpenXR/Source/OpenXRInput/Private/OpenXRInput.cpp
--- a/Engine/Plugins/Runtime/OpenXR/Source/OpenXRInput/Private/OpenXRInput.cpp
+++ b/Engine/Plugins/Runtime/OpenXR/Source/OpenXRInput/Private/OpenXRInput.cpp
@@ -298,9 +298,10 @@ void FOpenXRInputPlugin::FOpenXRInput::AddAction(XrActionSet ActionSet, const TA
 
 void FOpenXRInputPlugin::FOpenXRInput::Tick(float DeltaTime)
 {
-	if (OpenXRHMD->IsRunning() && OpenXRHMD->IsRendering())
+	XrSession Session = OpenXRHMD->GetSession();
+
+	if (OpenXRHMD->IsRunning())
 	{
-		XrSession Session = OpenXRHMD->GetSession();
 		if (!bActionsBound)
 		{
 			TArray<XrActionSet> BindActionSets;
@@ -312,29 +313,31 @@ void FOpenXRInputPlugin::FOpenXRInput::Tick(float DeltaTime)
 			SessionActionSetsAttachInfo.next = nullptr;
 			SessionActionSetsAttachInfo.countActionSets = BindActionSets.Num();
 			SessionActionSetsAttachInfo.actionSets = BindActionSets.GetData();
-			xrAttachSessionActionSets(Session, &SessionActionSetsAttachInfo);
+			XR_ENSURE(xrAttachSessionActionSets(Session, &SessionActionSetsAttachInfo));
 
 			bActionsBound = true;
 		}
+	}
+	else if (bActionsBound)
+	{
+		// If the session shut down, clean up.
+		bActionsBound = false;
+	}
 
+	if (OpenXRHMD->IsFocused())
+	{
 		XrActionsSyncInfo SyncInfo;
 		SyncInfo.type = XR_TYPE_ACTIONS_SYNC_INFO;
 		SyncInfo.next = nullptr;
 		SyncInfo.countActiveActionSets = ActionSets.Num();
 		SyncInfo.activeActionSets = ActionSets.GetData();
- 		XR_ENSURE(xrSyncActions(Session, &SyncInfo));
-	}
-	else if (bActionsBound && !OpenXRHMD->IsRunning())
-	{
-		// If the session shut down, clean up.
-		bActionsBound = false;
+		XR_ENSURE(xrSyncActions(Session, &SyncInfo));
 	}
-
 }
 
 void FOpenXRInputPlugin::FOpenXRInput::SendControllerEvents()
 {
-	if (!OpenXRHMD->IsRunning())
+	if (!bActionsBound)
 	{
 		return;
 	}
