Remove redundant second call to EnumerateExtensions in OpenXR plugin pre-init call.

diff --git a/Engine/Plugins/Runtime/OpenXR/Source/OpenXRHMD/Private/OpenXRHMD.cpp b/Engine/Plugins/Runtime/OpenXR/Source/OpenXRHMD/Private/OpenXRHMD.cpp
--- a/Engine/Plugins/Runtime/OpenXR/Source/OpenXRHMD/Private/OpenXRHMD.cpp
+++ b/Engine/Plugins/Runtime/OpenXR/Source/OpenXRHMD/Private/OpenXRHMD.cpp
@@ -237,8 +237,6 @@ bool FOpenXRHMDPlugin::PreInit()
 		return false;
 	}
 
-	EnumerateExtensions();
-
 	TArray<const char*> Extensions;
 	EnumerateExtensions();
 #ifdef XR_USE_GRAPHICS_API_D3D11
