Fail extension enumeration gracefully if no runtime is installed

diff --git a/Engine/Plugins/Runtime/OpenXR/Source/OpenXRHMD/Private/OpenXRHMD.cpp b/Engine/Plugins/Runtime/OpenXR/Source/OpenXRHMD/Private/OpenXRHMD.cpp
--- a/Engine/Plugins/Runtime/OpenXR/Source/OpenXRHMD/Private/OpenXRHMD.cpp
+++ b/Engine/Plugins/Runtime/OpenXR/Source/OpenXRHMD/Private/OpenXRHMD.cpp
@@ -120,7 +120,7 @@ private:
 	TSet<FString> AvailableExtensions;
 	TRefCountPtr<FOpenXRRenderBridge> RenderBridge;
 
-	void EnumerateExtensions();
+	bool EnumerateExtensions();
 	bool InitRenderBridge();
 };
 
@@ -156,21 +156,31 @@ uint64 FOpenXRHMDPlugin::GetGraphicsAdapterLuid()
 	return RenderBridge->GetGraphicsAdapterLuid();
 }
 
-void FOpenXRHMDPlugin::EnumerateExtensions()
+bool FOpenXRHMDPlugin::EnumerateExtensions()
 {
 	uint32_t ExtensionsCount = 0;
-	XR_ENSURE(xrEnumerateInstanceExtensionProperties(nullptr, 0, &ExtensionsCount, nullptr));
+	if (XR_FAILED(xrEnumerateInstanceExtensionProperties(nullptr, 0, &ExtensionsCount, nullptr)))
+	{
+		// If it fails this early that means there's no runtime installed
+		return false;
+	}
+
 	TArray<XrExtensionProperties> Properties;
 	Properties.SetNum(ExtensionsCount);
 	for (auto& Prop : Properties)
 	{
 		Prop = XrExtensionProperties{ XR_TYPE_EXTENSION_PROPERTIES };
 	}
-	XR_ENSURE(xrEnumerateInstanceExtensionProperties(nullptr, ExtensionsCount, &ExtensionsCount, Properties.GetData()));
-	for (const XrExtensionProperties& Prop : Properties)
+
+	if (XR_ENSURE(xrEnumerateInstanceExtensionProperties(nullptr, ExtensionsCount, &ExtensionsCount, Properties.GetData())))
 	{
-		AvailableExtensions.Add(Prop.extensionName);
+		for (const XrExtensionProperties& Prop : Properties)
+		{
+			AvailableExtensions.Add(Prop.extensionName);
+		}
+		return true;
 	}
+	return false;
 }
 
 bool FOpenXRHMDPlugin::InitRenderBridge()
@@ -238,7 +248,12 @@ bool FOpenXRHMDPlugin::PreInit()
 	}
 
 	TArray<const char*> Extensions;
-	EnumerateExtensions();
+	if (!EnumerateExtensions())
+	{
+		UE_LOG(LogHMD, Log, TEXT("Failed to enumerate extensions. Please check if you have an OpenXR runtime installed."));
+		return false;
+	}
+
 #ifdef XR_USE_GRAPHICS_API_D3D11
 	if (HasExtension(XR_KHR_D3D11_ENABLE_EXTENSION_NAME))
 	{
