Fixup OpenXR Vulkan support

diff --git a/Engine/Plugins/Runtime/OpenXR/Source/OpenXRHMD/Private/OpenXRHMD.cpp b/Engine/Plugins/Runtime/OpenXR/Source/OpenXRHMD/Private/OpenXRHMD.cpp
--- a/Engine/Plugins/Runtime/OpenXR/Source/OpenXRHMD/Private/OpenXRHMD.cpp
+++ b/Engine/Plugins/Runtime/OpenXR/Source/OpenXRHMD/Private/OpenXRHMD.cpp
@@ -87,6 +87,7 @@ public:
 
 	/** IHeadMountedDisplayModule implementation */
 	virtual TSharedPtr< class IXRTrackingSystem, ESPMode::ThreadSafe > CreateTrackingSystem() override;
+	virtual TSharedPtr< IHeadMountedDisplayVulkanExtensions, ESPMode::ThreadSafe > GetVulkanExtensions() override;
 	virtual uint64 GetGraphicsAdapterLuid() override;
 	virtual bool PreInit() override;
 
@@ -119,6 +120,7 @@ private:
 	XrSystemId System;
 	TSet<FString> AvailableExtensions;
 	TRefCountPtr<FOpenXRRenderBridge> RenderBridge;
+	TSharedPtr< IHeadMountedDisplayVulkanExtensions, ESPMode::ThreadSafe > VulkanExtensions;
 
 	bool EnumerateExtensions();
 	bool InitRenderBridge();
@@ -156,6 +158,21 @@ uint64 FOpenXRHMDPlugin::GetGraphicsAdapterLuid()
 	return RenderBridge->GetGraphicsAdapterLuid();
 }
 
+TSharedPtr< IHeadMountedDisplayVulkanExtensions, ESPMode::ThreadSafe > FOpenXRHMDPlugin::GetVulkanExtensions()
+{
+#if XR_USE_GRAPHICS_API_VULKAN
+	if (PreInit() && HasExtension(XR_KHR_VULKAN_ENABLE_EXTENSION_NAME))
+	{
+		if (!VulkanExtensions.IsValid())
+		{
+			VulkanExtensions = MakeShareable(new FOpenXRHMD::FVulkanExtensions(Instance, System));
+		}
+		return VulkanExtensions;
+	}
+#endif//XR_USE_GRAPHICS_API_VULKAN
+	return nullptr;
+}
+
 bool FOpenXRHMDPlugin::EnumerateExtensions()
 {
 	uint32_t ExtensionsCount = 0;
@@ -228,6 +245,9 @@ bool FOpenXRHMDPlugin::InitRenderBridge()
 
 bool FOpenXRHMDPlugin::PreInit()
 {
+	if (Instance)
+		return true;
+
 #if PLATFORM_WINDOWS
 #if PLATFORM_64BITS
 	FString BinariesPath = FPaths::EngineDir() / FString(TEXT("Binaries/ThirdParty/OpenXR/win64"));
@@ -359,12 +379,11 @@ bool FOpenXRHMD::FVulkanExtensions::GetVulkanInstanceExtensionsRequired(TArray<c
 		VulkanRHI::vkEnumerateInstanceExtensionProperties(nullptr, &PropertyCount, Properties.GetData());
 	}
 
-	TArray<char> Extensions;
 	{
 		uint32 ExtensionCount = 0;
-		xrGetVulkanInstanceExtensionsKHR(Instance, System, 0, &ExtensionCount, nullptr);
+		XR_ENSURE(xrGetVulkanInstanceExtensionsKHR(Instance, System, 0, &ExtensionCount, nullptr));
 		Extensions.SetNum(ExtensionCount);
-		xrGetVulkanInstanceExtensionsKHR(Instance, System, ExtensionCount, &ExtensionCount, Extensions.GetData());
+		XR_ENSURE(xrGetVulkanInstanceExtensionsKHR(Instance, System, ExtensionCount, &ExtensionCount, Extensions.GetData()));
 	}
 
 	ANSICHAR* Context = nullptr;
@@ -404,16 +423,15 @@ bool FOpenXRHMD::FVulkanExtensions::GetVulkanDeviceExtensionsRequired(VkPhysical
 		VulkanRHI::vkEnumerateDeviceExtensionProperties((VkPhysicalDevice)pPhysicalDevice, nullptr, &PropertyCount, Properties.GetData());
 	}
 
-	TArray<char> Extensions;
 	{
 		uint32 ExtensionCount = 0;
-		xrGetVulkanDeviceExtensionsKHR(Instance, System, 0, &ExtensionCount, nullptr);
-		Extensions.SetNum(ExtensionCount);
-		xrGetVulkanDeviceExtensionsKHR(Instance, System, ExtensionCount, &ExtensionCount, Extensions.GetData());
+		XR_ENSURE(xrGetVulkanDeviceExtensionsKHR(Instance, System, 0, &ExtensionCount, nullptr));
+		DeviceExtensions.SetNum(ExtensionCount);
+		XR_ENSURE(xrGetVulkanDeviceExtensionsKHR(Instance, System, ExtensionCount, &ExtensionCount, DeviceExtensions.GetData()));
 	}
 
 	ANSICHAR* Context = nullptr;
-	for (ANSICHAR* Tok = FCStringAnsi::Strtok(Extensions.GetData(), " ", &Context); Tok != nullptr; Tok = FCStringAnsi::Strtok(nullptr, " ", &Context))
+	for (ANSICHAR* Tok = FCStringAnsi::Strtok(DeviceExtensions.GetData(), " ", &Context); Tok != nullptr; Tok = FCStringAnsi::Strtok(nullptr, " ", &Context))
 	{
 		bool ExtensionFound = false;
 		for (int32 PropertyIndex = 0; PropertyIndex < Properties.Num(); PropertyIndex++)
diff --git a/Engine/Plugins/Runtime/OpenXR/Source/OpenXRHMD/Private/OpenXRHMD.h b/Engine/Plugins/Runtime/OpenXR/Source/OpenXRHMD/Private/OpenXRHMD.h
--- a/Engine/Plugins/Runtime/OpenXR/Source/OpenXRHMD/Private/OpenXRHMD.h
+++ b/Engine/Plugins/Runtime/OpenXR/Source/OpenXRHMD/Private/OpenXRHMD.h
@@ -51,6 +51,9 @@ public:
 	private:
 		XrInstance Instance;
 		XrSystemId System;
+
+		TArray<char> Extensions;
+		TArray<char> DeviceExtensions;
 	};
 
 	/** IXRTrackingSystem interface */
diff --git a/Engine/Plugins/Runtime/OpenXR/Source/OpenXRHMD/Private/OpenXRHMD_Swapchain.cpp b/Engine/Plugins/Runtime/OpenXR/Source/OpenXRHMD/Private/OpenXRHMD_Swapchain.cpp
--- a/Engine/Plugins/Runtime/OpenXR/Source/OpenXRHMD/Private/OpenXRHMD_Swapchain.cpp
+++ b/Engine/Plugins/Runtime/OpenXR/Source/OpenXRHMD/Private/OpenXRHMD_Swapchain.cpp
@@ -38,7 +38,7 @@ void FOpenXRSwapchain::IncrementSwapChainIndex_RHIThread(int64 Timeout)
 
 void FOpenXRSwapchain::ReleaseCurrentImage_RHIThread()
 {
-	check(IsInRenderingThread());
+	check(IsInRenderingThread() || IsInRHIThread());
 
 	if (!IsAcquired)
 		return;
@@ -64,8 +64,8 @@ uint8 GetNearestSupportedSwapchainFormat(XrSession InSession, uint8 RequestedFor
 		ToPlatformFormat = [](uint8 InFormat) { return GPixelFormats[InFormat].PlatformFormat; };
 	}
 
-	uint32_t FormatsCount;
-	xrEnumerateSwapchainFormats(InSession, 0, &FormatsCount, nullptr);
+	uint32_t FormatsCount = 0;
+	XR_ENSURE(xrEnumerateSwapchainFormats(InSession, 0, &FormatsCount, nullptr));
 
 	TArray<int64_t> Formats;
 	Formats.SetNum(FormatsCount);
@@ -271,13 +271,17 @@ FXRSwapChainPtr CreateSwapchain_OpenGL(XrSession InSession, uint8 Format, uint32
 #ifdef XR_USE_GRAPHICS_API_VULKAN
 FXRSwapChainPtr CreateSwapchain_Vulkan(XrSession InSession, uint8 Format, uint32 SizeX, uint32 SizeY, uint32 NumMips, uint32 NumSamples, uint32 Flags, uint32 TargetableTextureFlags)
 {
-	Format = GetNearestSupportedSwapchainFormat(InSession, Format);
+	TFunction<uint32(uint8)> ToPlatformFormat = [Flags](uint8 InFormat)
+	{
+		return UEToVkTextureFormat(GPixelFormats[InFormat].UnrealFormat, Flags & TexCreate_SRGB);
+	};
+	Format = GetNearestSupportedSwapchainFormat(InSession, Format, ToPlatformFormat);
 	if (!Format)
 	{
 		return nullptr;
 	}
 
-	XrSwapchain Swapchain = CreateSwapchain(InSession, GPixelFormats[Format].PlatformFormat, SizeX, SizeY, NumMips, NumSamples, TargetableTextureFlags);
+	XrSwapchain Swapchain = CreateSwapchain(InSession, ToPlatformFormat(Format), SizeX, SizeY, NumMips, NumSamples, TargetableTextureFlags);
 	if (!Swapchain)
 	{
 		return nullptr;
