Plumbing for depth composition layer in OpenXR - forced off currently until we figure out the right way to handle the size disparity between the renderer's buffers and the swapchain/OpenXR provided depth chain.

diff --git a/Engine/Plugins/Runtime/OpenXR/Source/OpenXRHMD/Private/OpenXRHMD.cpp b/Engine/Plugins/Runtime/OpenXR/Source/OpenXRHMD/Private/OpenXRHMD.cpp
--- a/Engine/Plugins/Runtime/OpenXR/Source/OpenXRHMD/Private/OpenXRHMD.cpp
+++ b/Engine/Plugins/Runtime/OpenXR/Source/OpenXRHMD/Private/OpenXRHMD.cpp
@@ -19,6 +19,7 @@
 #include "Misc/CString.h"
 #include "ClearQuad.h"
 #include "XRThreadUtils.h"
+#include "RenderUtils.h"
 #include "PipelineStateCache.h"
 #include "Slate/SceneViewport.h"
 #include "Engine/GameEngine.h"
@@ -138,11 +139,12 @@ TSharedPtr< class IXRTrackingSystem, ESPMode::ThreadSafe > FOpenXRHMDPlugin::Cre
 		}
 	}
 
-	auto OpenXRHMD = FSceneViewExtensions::NewExtension<FOpenXRHMD>(Instance, System, RenderBridge);
+	auto OpenXRHMD = FSceneViewExtensions::NewExtension<FOpenXRHMD>(Instance, System, RenderBridge, HasExtension(XR_KHR_COMPOSITION_LAYER_DEPTH_EXTENSION_NAME));
 	if( OpenXRHMD->IsInitialized() )
 	{
 		return OpenXRHMD;
 	}
+
 	return nullptr;
 }
 
@@ -299,12 +301,10 @@ bool FOpenXRHMDPlugin::PreInit()
 	}
 #endif
 
-#if 0
 	if (HasExtension(XR_KHR_COMPOSITION_LAYER_DEPTH_EXTENSION_NAME))
 	{
 		Extensions.Add(XR_KHR_COMPOSITION_LAYER_DEPTH_EXTENSION_NAME);
 	}
-#endif
 
 	if (HasExtension(XR_VARJO_QUAD_VIEWS_EXTENSION_NAME))
 	{
@@ -740,7 +740,7 @@ bool FOpenXRHMD::IsActiveThisFrame(class FViewport* InViewport) const
 	return GEngine && GEngine->IsStereoscopic3D(InViewport);
 }
 
-FOpenXRHMD::FOpenXRHMD(const FAutoRegister& AutoRegister, XrInstance InInstance, XrSystemId InSystem, TRefCountPtr<FOpenXRRenderBridge>& InRenderBridge)
+FOpenXRHMD::FOpenXRHMD(const FAutoRegister& AutoRegister, XrInstance InInstance, XrSystemId InSystem, TRefCountPtr<FOpenXRRenderBridge>& InRenderBridge, bool InDepthExtensionSupported)
 	: FHeadMountedDisplayBase(nullptr)
 	, FSceneViewExtensionBase(AutoRegister)
 	, bStereoEnabled(false)
@@ -748,6 +748,8 @@ FOpenXRHMD::FOpenXRHMD(const FAutoRegister& AutoRegister, XrInstance InInstance,
 	, bIsReady(false)
 	, bIsRendering(false)
 	, bRunRequested(false)
+	, bDepthExtensionSupported(InDepthExtensionSupported)
+	, bNeedReAllocatedDepth(InDepthExtensionSupported)
 	, CurrentSessionState(XR_SESSION_STATE_UNKNOWN)
 	, Instance(InInstance)
 	, System(InSystem)
@@ -968,21 +970,24 @@ bool FOpenXRHMD::AllocateRenderTargetTexture(uint32 Index, uint32 SizeX, uint32
 	
 	// Grab the presentation texture out of the swapchain.
 	OutTargetableTexture = OutShaderResourceTexture = (FTexture2DRHIRef&)Swapchain->GetTextureRef();
-#if 0
-	// Allocate the depth buffer swapchain while we're here.
-	DepthSwapchain = RenderBridge->CreateSwapchain(Session, PF_DepthStencil, SizeX, SizeY, NumMips, NumSamples, 0, TexCreate_DepthStencilTargetable);
-	if (!DepthSwapchain)
+
+	if (bDepthExtensionSupported)
 	{
-		return false;
+		// Allocate the depth buffer swapchain while we're here.
+		DepthSwapchain = RenderBridge->CreateSwapchain(Session, PF_DepthStencil, SizeX, SizeY, NumMips, NumSamples, 0, TexCreate_DepthStencilTargetable);
+		if (!DepthSwapchain)
+		{
+			return false;
+		}
+		bNeedReAllocatedDepth = false;
 	}
-#endif
+
 	return true;
 }
 
-#if 0
 bool FOpenXRHMD::AllocateDepthTexture(uint32 Index, uint32 SizeX, uint32 SizeY, uint8 Format, uint32 NumMips, uint32 InTexFlags, uint32 TargetableTextureFlags, FTexture2DRHIRef& OutTargetableTexture, FTexture2DRHIRef& OutShaderResourceTexture, uint32 NumSamples)
 {
-	if (!DepthSwapchain || DepthSwapchain->GetTexture2D()->GetSizeX() != SizeX || DepthSwapchain->GetTexture2D()->GetSizeY() != SizeY)
+	if (!DepthSwapchain.IsValid())
 	{
 		return false;
 	}
@@ -991,7 +996,6 @@ bool FOpenXRHMD::AllocateDepthTexture(uint32 Index, uint32 SizeX, uint32 SizeY,
 
 	return true;
 }
-#endif
 
 void FOpenXRHMD::OnBeginRendering_RenderThread(FRHICommandListImmediate& RHICmdList, FSceneViewFamily& ViewFamily)
 {
@@ -1004,9 +1008,11 @@ void FOpenXRHMD::OnBeginRendering_RenderThread(FRHICommandListImmediate& RHICmdL
 		bIsRendering = true;
 
 		Swapchain->IncrementSwapChainIndex_RHIThread(FrameStateRHI.predictedDisplayPeriod);
-#if 0
-		DepthSwapchain->IncrementSwapChainIndex_RHIThread(FrameStateRHI.predictedDisplayPeriod);
-#endif
+		if (bDepthExtensionSupported)
+		{
+			ensure(DepthSwapchain != nullptr);
+			DepthSwapchain->IncrementSwapChainIndex_RHIThread(FrameStateRHI.predictedDisplayPeriod);
+		}
 	}
 
 	FrameStateRHI = FrameState;
@@ -1046,17 +1052,18 @@ void FOpenXRHMD::OnBeginRendering_RenderThread(FRHICommandListImmediate& RHICmdL
 			}
 		};
 
-#if 0
-		DepthLayer.type = XR_TYPE_COMPOSITION_LAYER_DEPTH_INFO_KHR;
-		DepthLayer.next = nullptr;
-		DepthLayer.subImage.swapchain = static_cast<FOpenXRSwapchain*>(GetDepthSwapchain())->GetHandle();
-		DepthLayer.subImage.imageArrayIndex = 0;
-		DepthLayer.subImage.imageRect = Projection.subImage.imageRect;
-		DepthLayer.minDepth = 1.0f;
-		DepthLayer.maxDepth = 0.0f;
-		DepthLayer.nearZ = NearZ;
-		DepthLayer.farZ = FLT_MAX;
-#endif
+		if (bDepthExtensionSupported)
+		{
+			DepthLayer.type = XR_TYPE_COMPOSITION_LAYER_DEPTH_INFO_KHR;
+			DepthLayer.next = nullptr;
+			DepthLayer.subImage.swapchain = static_cast<FOpenXRSwapchain*>(GetDepthSwapchain())->GetHandle();
+			DepthLayer.subImage.imageArrayIndex = 0;
+			DepthLayer.subImage.imageRect = Projection.subImage.imageRect;
+			DepthLayer.minDepth = 1.0f;
+			DepthLayer.maxDepth = 0.0f;
+			DepthLayer.nearZ = NearZ;
+			DepthLayer.farZ = FLT_MAX;
+		}
 
 		OffsetX += Config.recommendedImageRectWidth;
 	}
@@ -1244,9 +1251,11 @@ void FOpenXRHMD::FinishRendering()
 	if (bIsRendering)
 	{
 		Swapchain->ReleaseCurrentImage_RHIThread();
-#if 0
-		DepthSwapchain->ReleaseCurrentImage_RHIThread();
-#endif
+
+		if (bDepthExtensionSupported)
+		{
+			DepthSwapchain->ReleaseCurrentImage_RHIThread();
+		}
 
 		XrFrameEndInfo EndInfo;
 		XrCompositionLayerBaseHeader* Headers[1] = { reinterpret_cast<XrCompositionLayerBaseHeader*>(&Layer) };
@@ -1276,6 +1285,11 @@ FIntPoint FOpenXRHMD::GetIdealRenderTargetSize() const
 		Size.X += (int)Config.recommendedImageRectWidth;
 		Size.Y = FMath::Max(Size.Y, (int)Config.recommendedImageRectHeight);
 	}
+
+	// We always prefer the nearest multiple of 4 for our buffer sizes. Make sure we round up here,
+	// so we're consistent with the rest of the engine in creating our buffers.
+	QuantizeSceneBufferSize(Size, Size);
+
 	return Size;
 }
 
diff --git a/Engine/Plugins/Runtime/OpenXR/Source/OpenXRHMD/Private/OpenXRHMD.h b/Engine/Plugins/Runtime/OpenXR/Source/OpenXRHMD/Private/OpenXRHMD.h
--- a/Engine/Plugins/Runtime/OpenXR/Source/OpenXRHMD/Private/OpenXRHMD.h
+++ b/Engine/Plugins/Runtime/OpenXR/Source/OpenXRHMD/Private/OpenXRHMD.h
@@ -152,10 +152,8 @@ public:
 	/** IStereoRenderTargetManager */
 	virtual bool ShouldUseSeparateRenderTarget() const override { return true; }
 	virtual bool AllocateRenderTargetTexture(uint32 Index, uint32 SizeX, uint32 SizeY, uint8 Format, uint32 NumMips, uint32 Flags, uint32 TargetableTextureFlags, FTexture2DRHIRef& OutTargetableTexture, FTexture2DRHIRef& OutShaderResourceTexture, uint32 NumSamples = 1) override;
-#if 0
-	virtual bool NeedReAllocateDepthTexture(const TRefCountPtr<IPooledRenderTarget>& DepthTarget) override final { return false; }
+	virtual bool NeedReAllocateDepthTexture(const TRefCountPtr<IPooledRenderTarget>& DepthTarget) override final { return bNeedReAllocatedDepth; }
 	virtual bool AllocateDepthTexture(uint32 Index, uint32 SizeX, uint32 SizeY, uint8 Format, uint32 NumMips, uint32 InTexFlags, uint32 TargetableTextureFlags, FTexture2DRHIRef& OutTargetableTexture, FTexture2DRHIRef& OutShaderResourceTexture, uint32 NumSamples = 1) override final;
-#endif
 
 	virtual FXRRenderBridge* GetActiveRenderBridge_GameThread(bool bUseSeparateRenderTarget) override;
 
@@ -165,7 +163,7 @@ public:
 
 public:
 	/** Constructor */
-	FOpenXRHMD(const FAutoRegister&, XrInstance InInstance, XrSystemId InSystem, TRefCountPtr<FOpenXRRenderBridge>& InRenderBridge);
+	FOpenXRHMD(const FAutoRegister&, XrInstance InInstance, XrSystemId InSystem, TRefCountPtr<FOpenXRRenderBridge>& InRenderBridge, bool InDepthExtensionSupported);
 
 	/** Destructor */
 	virtual ~FOpenXRHMD();
@@ -179,9 +177,8 @@ public:
 	OPENXRHMD_API int32 AddActionDevice(XrAction Action);
 
 	FXRSwapChain* GetSwapchain() { return Swapchain.Get(); }
-#if 0
 	FXRSwapChain* GetDepthSwapchain() { return DepthSwapchain.Get(); }
-#endif
+
 	XrInstance GetInstance() { return Instance; }
 	XrSystemId GetSystem() { return System; }
 	XrSession GetSession() { return Session; }
@@ -196,6 +193,8 @@ private:
 	bool					bIsReady;
 	bool					bIsRendering;
 	bool					bRunRequested;
+	bool					bDepthExtensionSupported;
+	bool					bNeedReAllocatedDepth;
 	
 	XrSessionState			CurrentSessionState;
 
@@ -226,7 +225,5 @@ private:
 	IRendererModule*		RendererModule;
 
 	FXRSwapChainPtr			Swapchain;
-#if 0
 	FXRSwapChainPtr			DepthSwapchain;
-#endif
 };
diff --git a/Engine/Plugins/Runtime/OpenXR/Source/OpenXRHMD/Private/OpenXRHMD_Swapchain.cpp b/Engine/Plugins/Runtime/OpenXR/Source/OpenXRHMD/Private/OpenXRHMD_Swapchain.cpp
--- a/Engine/Plugins/Runtime/OpenXR/Source/OpenXRHMD/Private/OpenXRHMD_Swapchain.cpp
+++ b/Engine/Plugins/Runtime/OpenXR/Source/OpenXRHMD/Private/OpenXRHMD_Swapchain.cpp
@@ -187,14 +187,16 @@ FXRSwapChainPtr CreateSwapchain_D3D11(XrSession InSession, uint8 Format, uint32
 		return nullptr;
 	}
 
+	const bool bDepthStencil = (PF_DepthStencil == Format);
+
 	FD3D11DynamicRHI* DynamicRHI = static_cast<FD3D11DynamicRHI*>(GDynamicRHI);
 	TArray<FTextureRHIRef> TextureChain;
 	// @todo: Once things settle down, the chain target will be created below in the CreateXRSwapChain call, via an RHI "CreateAliasedTexture" call.
 	TArray<XrSwapchainImageD3D11KHR> Images = EnumerateImages<XrSwapchainImageD3D11KHR>(Swapchain, XR_TYPE_SWAPCHAIN_IMAGE_D3D11_KHR);
-	FTextureRHIRef ChainTarget = static_cast<FTextureRHIRef>(DynamicRHI->RHICreateTexture2DFromResource(GPixelFormats[Format].UnrealFormat, TargetableTextureFlags, FClearValueBinding::Black, Images[0].texture));
+	FTextureRHIRef ChainTarget = static_cast<FTextureRHIRef>(DynamicRHI->RHICreateTexture2DFromResource(GPixelFormats[Format].UnrealFormat, TargetableTextureFlags, bDepthStencil ? FClearValueBinding::DepthFar : FClearValueBinding::Black, Images[0].texture));
 	for (const auto& Image : Images)
 	{
-		TextureChain.Add(static_cast<FTextureRHIRef>(DynamicRHI->RHICreateTexture2DFromResource(GPixelFormats[Format].UnrealFormat, TargetableTextureFlags, FClearValueBinding::Black, Image.texture)));
+		TextureChain.Add(static_cast<FTextureRHIRef>(DynamicRHI->RHICreateTexture2DFromResource(GPixelFormats[Format].UnrealFormat, TargetableTextureFlags, bDepthStencil ? FClearValueBinding::DepthFar : FClearValueBinding::Black, Image.texture)));
 	}
 
 	return CreateXRSwapChain<FOpenXRSwapchain>(MoveTemp(TextureChain), ChainTarget, Swapchain);
