Add plumbing plus CVar for OpenXR validation layer loading.

diff --git a/Engine/Plugins/Runtime/OpenXR/Source/OpenXRHMD/Private/OpenXRHMD.cpp b/Engine/Plugins/Runtime/OpenXR/Source/OpenXRHMD/Private/OpenXRHMD.cpp
--- a/Engine/Plugins/Runtime/OpenXR/Source/OpenXRHMD/Private/OpenXRHMD.cpp
+++ b/Engine/Plugins/Runtime/OpenXR/Source/OpenXRHMD/Private/OpenXRHMD.cpp
@@ -29,6 +29,14 @@
 #include "Editor/UnrealEd/Classes/Editor/EditorEngine.h"
 #endif
 
+static TAutoConsoleVariable<int32> CVarEnableOpenXRValidationLayer(
+	TEXT("xr.EnableOpenXRValidationLayer"),
+	0,
+	TEXT("If true, enables the OpenXR validation layer, which will provide extended validation of\nOpenXR API calls. This should only be used for debugging purposes.\n")
+	TEXT("Changes will only take effect in new game/editor instances - can't be changed at runtime.\n"),
+	ECVF_Default);		// @todo: Should we specify ECVF_Cheat here so this doesn't show up in release builds?
+
+
 namespace {
 	/** Helper function for acquiring the appropriate FSceneViewport */
 	FSceneViewport* FindSceneViewport()
@@ -114,16 +122,19 @@ public:
 	virtual bool IsHMDConnected() override { return true; }
 
 	bool HasExtension(const char* Name) const { return AvailableExtensions.Contains(Name); }
+	bool HasLayer(const char* Name) const { return AvailableLayers.Contains(Name); }
 
 private:
 	void *LoaderHandle;
 	XrInstance Instance;
 	XrSystemId System;
 	TSet<FString> AvailableExtensions;
+	TSet<FString> AvailableLayers;
 	TRefCountPtr<FOpenXRRenderBridge> RenderBridge;
 	TSharedPtr< IHeadMountedDisplayVulkanExtensions, ESPMode::ThreadSafe > VulkanExtensions;
 
 	bool EnumerateExtensions();
+	bool EnumerateLayers();
 	bool InitRenderBridge();
 };
 
@@ -202,6 +213,40 @@ bool FOpenXRHMDPlugin::EnumerateExtensions()
 	return false;
 }
 
+bool FOpenXRHMDPlugin::EnumerateLayers()
+{
+	uint32 LayerPropertyCount = 0;
+	if (XR_FAILED(xrEnumerateApiLayerProperties(0, &LayerPropertyCount, nullptr)))
+	{
+		// As per EnumerateExtensions - a failure here means no runtime installed.
+		return false;
+	}
+
+	if (!LayerPropertyCount)
+	{
+		// It's still legit if we have no layers, so early out here (and return success) if so.
+		return true;
+	}
+
+	TArray<XrApiLayerProperties> LayerProperties;
+	LayerProperties.SetNum(LayerPropertyCount);
+	for (auto& Prop : LayerProperties)
+	{
+		Prop = XrApiLayerProperties{ XR_TYPE_API_LAYER_PROPERTIES };
+	}
+
+	if (XR_ENSURE(xrEnumerateApiLayerProperties(LayerPropertyCount, &LayerPropertyCount, LayerProperties.GetData())))
+	{
+		for (const auto& Prop : LayerProperties)
+		{
+			AvailableLayers.Add(Prop.layerName);
+		}
+		return true;
+	}
+
+	return false;
+}
+
 bool FOpenXRHMDPlugin::InitRenderBridge()
 {
 	FString RHIString = FApp::GetGraphicsRHI();
@@ -272,7 +317,13 @@ bool FOpenXRHMDPlugin::PreInit()
 	TArray<const char*> Extensions;
 	if (!EnumerateExtensions())
 	{
-		UE_LOG(LogHMD, Log, TEXT("Failed to enumerate extensions. Please check if you have an OpenXR runtime installed."));
+		UE_LOG(LogHMD, Log, TEXT("Failed to enumerate extensions. Please check that you have a valid OpenXR runtime installed."));
+		return false;
+	}
+
+	if (!EnumerateLayers())
+	{
+		UE_LOG(LogHMD, Log, TEXT("Failed to enumerate API layers. Please check that you have a valid OpenXR runtime installed."));
 		return false;
 	}
 
@@ -311,6 +362,14 @@ bool FOpenXRHMDPlugin::PreInit()
 		Extensions.Add(XR_VARJO_QUAD_VIEWS_EXTENSION_NAME);
 	}
 
+	// Enable layers, if specified by CVar
+	const bool bEnableOpenXRValidationLayer = CVarEnableOpenXRValidationLayer.GetValueOnAnyThread() != 0;
+	TArray<const char*> Layers;
+	if (bEnableOpenXRValidationLayer && HasLayer("XR_APILAYER_LUNARG_core_validation"))
+	{
+		Layers.Add("XR_APILAYER_LUNARG_core_validation");
+	}
+
 	// Engine registration can be disabled via console var.
 	auto* CVarDisableEngineAndAppRegistration = IConsoleManager::Get().FindTConsoleVariableDataInt(TEXT("r.DisableEngineAndAppRegistration"));
 	bool bDisableEngineRegistration = (CVarDisableEngineAndAppRegistration && CVarDisableEngineAndAppRegistration->GetValueOnAnyThread() != 0);
@@ -335,10 +394,13 @@ bool FOpenXRHMDPlugin::PreInit()
 	FTCHARToUTF8_Convert::Convert(Info.applicationInfo.engineName, XR_MAX_ENGINE_NAME_SIZE, *EngineName, EngineName.Len() + 1);
 	Info.applicationInfo.engineVersion = (uint32)(FEngineVersion::Current().GetMinor() << 16 | FEngineVersion::Current().GetPatch());
 	Info.applicationInfo.apiVersion = XR_CURRENT_API_VERSION;
-	Info.enabledApiLayerCount = 0;
-	Info.enabledApiLayerNames = nullptr;
+
+	Info.enabledApiLayerCount = Layers.Num();
+	Info.enabledApiLayerNames = Layers.GetData();
+
 	Info.enabledExtensionCount = Extensions.Num();
 	Info.enabledExtensionNames = Extensions.GetData();
+
 	XrResult rs = xrCreateInstance(&Info, &Instance);
 	if (XR_FAILED(rs))
 	{
