Fix multiple xrEndFrame calls in frames when PIE, editor dropdown clicked.

diff --git a/Engine/Plugins/Runtime/OpenXR/Source/OpenXRHMD/Private/OpenXRHMD.cpp b/Engine/Plugins/Runtime/OpenXR/Source/OpenXRHMD/Private/OpenXRHMD.cpp
--- a/Engine/Plugins/Runtime/OpenXR/Source/OpenXRHMD/Private/OpenXRHMD.cpp
+++ b/Engine/Plugins/Runtime/OpenXR/Source/OpenXRHMD/Private/OpenXRHMD.cpp
@@ -1066,14 +1066,28 @@ void FOpenXRHMD::OnBeginRendering_RenderThread(FRHICommandListImmediate& RHICmdL
 		XrFrameBeginInfo BeginInfo;
 		BeginInfo.type = XR_TYPE_FRAME_BEGIN_INFO;
 		BeginInfo.next = nullptr;
-		xrBeginFrame(Session, &BeginInfo);
-		bIsRendering = true;
+		XrResult Result = xrBeginFrame(Session, &BeginInfo);
+		if (XR_SUCCEEDED(Result))
+		{
+			bIsRendering = true;
 
-		Swapchain->IncrementSwapChainIndex_RHIThread(FrameStateRHI.predictedDisplayPeriod);
-		if (bDepthExtensionSupported)
+			Swapchain->IncrementSwapChainIndex_RHIThread(FrameStateRHI.predictedDisplayPeriod);
+			if (bDepthExtensionSupported)
+			{
+				ensure(DepthSwapchain != nullptr);
+				DepthSwapchain->IncrementSwapChainIndex_RHIThread(FrameStateRHI.predictedDisplayPeriod);
+			}
+		}
+		else
 		{
-			ensure(DepthSwapchain != nullptr);
-			DepthSwapchain->IncrementSwapChainIndex_RHIThread(FrameStateRHI.predictedDisplayPeriod);
+			static bool bLoggedBeginFrameFailure = false;
+			if (!bLoggedBeginFrameFailure)
+			{
+				char Error[XR_MAX_RESULT_STRING_SIZE] = { '\0' };
+				xrResultToString(XR_NULL_HANDLE, Result, Error);
+				UE_LOG(LogHMD, Error, TEXT("Unexpected error on xrBeginFrame. Error code was %s."), Error);
+				bLoggedBeginFrameFailure = true;
+			}
 		}
 	}
 
@@ -1331,6 +1345,8 @@ void FOpenXRHMD::FinishRendering()
 
 		// Ignore invalid call order for now, we will recover on the next frame
 		ensure(XR_SUCCEEDED(Result) || Result == XR_ERROR_CALL_ORDER_INVALID);
+
+		bIsRendering = false;
 	}
 }
 
