Any errors returned in OpenXR that we log will display garbage since the runtime returns UTF8 but our logging code is using Unicode. This change:
1. Renames the internal conversion macro in OpenXRHMDPrivate.h to not mirror the OpenXR API naming (and also make the name consistent with UE coding standards)
2. Replaces all uses of the OpenXR API with the UE version (which returns the correct type, and is also neater)

diff --git a/Engine/Plugins/Runtime/OpenXR/Source/OpenXRHMD/Private/OpenXRHMD.cpp b/Engine/Plugins/Runtime/OpenXR/Source/OpenXRHMD/Private/OpenXRHMD.cpp
--- a/Engine/Plugins/Runtime/OpenXR/Source/OpenXRHMD/Private/OpenXRHMD.cpp
+++ b/Engine/Plugins/Runtime/OpenXR/Source/OpenXRHMD/Private/OpenXRHMD.cpp
@@ -404,9 +404,7 @@ bool FOpenXRHMDPlugin::PreInit()
 	XrResult rs = xrCreateInstance(&Info, &Instance);
 	if (XR_FAILED(rs))
 	{
-		char error[XR_MAX_RESULT_STRING_SIZE] = { '\0' };
-		xrResultToString(XR_NULL_HANDLE, rs, error);
-		UE_LOG(LogHMD, Log, TEXT("Failed to create an OpenXR instance, result is %s. Please check if you have an OpenXR runtime installed."), error);
+		UE_LOG(LogHMD, Log, TEXT("Failed to create an OpenXR instance, result is %s. Please check if you have an OpenXR runtime installed."), OpenXRResultToString(rs));
 		return false;
 	}
 
@@ -417,9 +415,7 @@ bool FOpenXRHMDPlugin::PreInit()
 	rs = xrGetSystem(Instance, &SystemInfo, &System);
 	if (XR_FAILED(rs))
 	{
-		char error[XR_MAX_RESULT_STRING_SIZE] = { '\0' };
-		xrResultToString(XR_NULL_HANDLE, rs, error);
-		UE_LOG(LogHMD, Log, TEXT("Failed to get an OpenXR system, result is %s. Please check that your runtime supports VR headsets."), error);
+		UE_LOG(LogHMD, Log, TEXT("Failed to get an OpenXR system, result is %s. Please check that your runtime supports VR headsets."), OpenXRResultToString(rs));
 		return false;
 	}
 
@@ -1083,9 +1079,7 @@ void FOpenXRHMD::OnBeginRendering_RenderThread(FRHICommandListImmediate& RHICmdL
 			static bool bLoggedBeginFrameFailure = false;
 			if (!bLoggedBeginFrameFailure)
 			{
-				char Error[XR_MAX_RESULT_STRING_SIZE] = { '\0' };
-				xrResultToString(XR_NULL_HANDLE, Result, Error);
-				UE_LOG(LogHMD, Error, TEXT("Unexpected error on xrBeginFrame. Error code was %s."), Error);
+				UE_LOG(LogHMD, Error, TEXT("Unexpected error on xrBeginFrame. Error code was %s."), OpenXRResultToString(Result));
 				bLoggedBeginFrameFailure = true;
 			}
 		}
diff --git a/Engine/Plugins/Runtime/OpenXR/Source/OpenXRHMD/Private/OpenXRHMDPrivate.h b/Engine/Plugins/Runtime/OpenXR/Source/OpenXRHMD/Private/OpenXRHMDPrivate.h
--- a/Engine/Plugins/Runtime/OpenXR/Source/OpenXRHMD/Private/OpenXRHMDPrivate.h
+++ b/Engine/Plugins/Runtime/OpenXR/Source/OpenXRHMD/Private/OpenXRHMDPrivate.h
@@ -15,7 +15,7 @@
 #include <openxr/openxr_reflection.h>
 
 #define XR_ENUM_CASE_STR(name, val) case name: return TEXT(#name);
-constexpr const TCHAR* xrResultToString(XrResult e)
+constexpr const TCHAR* OpenXRResultToString(XrResult e)
 {
 	switch (e)
 	{
@@ -27,7 +27,7 @@ constexpr const TCHAR* xrResultToString(XrResult e)
 #if DO_CHECK
 #define XR_ENSURE(x) [] (XrResult Result) \
 	{ \
-		return ensureMsgf(XR_SUCCEEDED(Result), TEXT("OpenXR call failed with result %s"), xrResultToString(Result)); \
+		return ensureMsgf(XR_SUCCEEDED(Result), TEXT("OpenXR call failed with result %s"), OpenXRResultToString(Result)); \
 	} (x)
 #else
 #define XR_ENSURE(x) XR_SUCCEEDED(x)
