Enable OpenXR validation layer, fix runtime vaildation errors.

diff --git a/Engine/Plugins/Runtime/OpenXR/Source/OpenXRHMD/Private/OpenXRHMD.cpp b/Engine/Plugins/Runtime/OpenXR/Source/OpenXRHMD/Private/OpenXRHMD.cpp
--- a/Engine/Plugins/Runtime/OpenXR/Source/OpenXRHMD/Private/OpenXRHMD.cpp
+++ b/Engine/Plugins/Runtime/OpenXR/Source/OpenXRHMD/Private/OpenXRHMD.cpp
@@ -362,7 +362,14 @@ bool FOpenXRHMDPlugin::PreInit()
 		Extensions.Add(XR_VARJO_QUAD_VIEWS_EXTENSION_NAME);
 	}
 
-	// Enable layers, if specified by CVar
+	// Enable layers, if specified by CVar.
+	// Note: For the validation layer to work on Windows (as of latest OpenXR runtime, August 2019), the following are required:
+	//   1. Download and build the OpenXR SDK from [REDACTED_PROFILE_URL]/OpenXR-SDK-Source (follow instructions at [REDACTED_PROFILE_URL]/OpenXR-SDK-Source/blob/master/BUILDING.md)
+	//	 2. Add a registry key under HKEY_LOCAL_MACHINE\SOFTWARE\Khronos\OpenXR\1\ApiLayers\Explicit, containing the path to the manifest file
+	//      (e.g. C:\OpenXR-SDK-Source-master\build\win64\src\api_layers\XrApiLayer_core_validation.json) <-- this file is downloaded as part of the SDK source, above
+	//   3. Copy the DLL from the build target at, for example, C:\OpenXR-SDK-Source-master\build\win64\src\api_layers\XrApiLayer_core_validation.dll to 
+	//      somewhere in your system path (e.g. c:\windows\system32); the OpenXR loader currently doesn't use the path the json file is in (this is a bug)
+
 	const bool bEnableOpenXRValidationLayer = CVarEnableOpenXRValidationLayer.GetValueOnAnyThread() != 0;
 	TArray<const char*> Layers;
 	if (bEnableOpenXRValidationLayer && HasLayer("XR_APILAYER_LUNARG_core_validation"))
@@ -834,6 +841,9 @@ FOpenXRHMD::FOpenXRHMD(const FAutoRegister& AutoRegister, XrInstance InInstance,
 		TArray<XrViewConfigurationType> Types;
 		XR_ENSURE(xrEnumerateViewConfigurations(Instance, System, 0, &ConfigurationCount, nullptr));
 		Types.SetNum(ConfigurationCount);
+		// Fill the initial array with valid enum types (this will fail in the validation layer otherwise).
+		for (auto & TypeIter : Types)
+			TypeIter = XR_VIEW_CONFIGURATION_TYPE_PRIMARY_MONO;
 		XR_ENSURE(xrEnumerateViewConfigurations(Instance, System, ConfigurationCount, &ConfigurationCount, Types.GetData()));
 
 		// Ensure the configuration type we want is provided
@@ -884,18 +894,21 @@ bool FOpenXRHMD::OnStereoStartup()
 
 	FlushRenderingCommands();
 
-	uint32_t referenceSpacesCount;
-	XR_ENSURE(xrEnumerateReferenceSpaces(Session, 0, &referenceSpacesCount, nullptr));
+	uint32_t ReferenceSpacesCount;
+	XR_ENSURE(xrEnumerateReferenceSpaces(Session, 0, &ReferenceSpacesCount, nullptr));
 
-	TArray<XrReferenceSpaceType> spaces;
-	spaces.SetNum(referenceSpacesCount);
-	XR_ENSURE(xrEnumerateReferenceSpaces(Session, (uint32_t)spaces.Num(), &referenceSpacesCount, spaces.GetData()));
-	ensure(referenceSpacesCount == spaces.Num());
+	TArray<XrReferenceSpaceType> Spaces;
+	Spaces.SetNum(ReferenceSpacesCount);
+	// Initialize spaces array with valid enum values (avoid triggering validation error).
+	for (auto & SpaceIter : Spaces)
+		SpaceIter = XR_REFERENCE_SPACE_TYPE_VIEW;
+	XR_ENSURE(xrEnumerateReferenceSpaces(Session, (uint32_t)Spaces.Num(), &ReferenceSpacesCount, Spaces.GetData()));
+	ensure(ReferenceSpacesCount == Spaces.Num());
 
 	XrSpace HmdSpace = XR_NULL_HANDLE;
 	XrReferenceSpaceCreateInfo SpaceInfo;
 
-	ensure(spaces.Contains(XR_REFERENCE_SPACE_TYPE_VIEW));
+	ensure(Spaces.Contains(XR_REFERENCE_SPACE_TYPE_VIEW));
 	SpaceInfo.type = XR_TYPE_REFERENCE_SPACE_CREATE_INFO;
 	SpaceInfo.next = nullptr;
 	SpaceInfo.referenceSpaceType = XR_REFERENCE_SPACE_TYPE_VIEW;
@@ -903,12 +916,12 @@ bool FOpenXRHMD::OnStereoStartup()
 	XR_ENSURE(xrCreateReferenceSpace(Session, &SpaceInfo, &HmdSpace));
 	ActionSpaces[HMDDeviceId].Space = HmdSpace;
 
-	ensure(spaces.Contains(XR_REFERENCE_SPACE_TYPE_LOCAL));
+	ensure(Spaces.Contains(XR_REFERENCE_SPACE_TYPE_LOCAL));
 	SpaceInfo.referenceSpaceType = XR_REFERENCE_SPACE_TYPE_LOCAL;
 	XR_ENSURE(xrCreateReferenceSpace(Session, &SpaceInfo, &LocalSpace));
 
 	// Prefer a stage space over a local space
-	if (spaces.Contains(XR_REFERENCE_SPACE_TYPE_STAGE))
+	if (Spaces.Contains(XR_REFERENCE_SPACE_TYPE_STAGE))
 	{
 		TrackingSpaceType = XR_REFERENCE_SPACE_TYPE_STAGE;
 		SpaceInfo.referenceSpaceType = TrackingSpaceType;
