Fix missing cleanup of session if PIE is exited before session becomes active. This fixes a crash if PIE is activated more than once without wearing the headset.

diff --git a/Engine/Plugins/Runtime/OpenXR/Source/OpenXRHMD/Private/OpenXRHMD.cpp b/Engine/Plugins/Runtime/OpenXR/Source/OpenXRHMD/Private/OpenXRHMD.cpp
--- a/Engine/Plugins/Runtime/OpenXR/Source/OpenXRHMD/Private/OpenXRHMD.cpp
+++ b/Engine/Plugins/Runtime/OpenXR/Source/OpenXRHMD/Private/OpenXRHMD.cpp
@@ -879,6 +879,14 @@ FOpenXRHMD::FOpenXRHMD(const FAutoRegister& AutoRegister, XrInstance InInstance,
 	ensure(ActionSpaces.Emplace(XR_NULL_HANDLE) == HMDDeviceId);
 }
 
+FOpenXRHMD::~FOpenXRHMD()
+{
+	if (Session)
+	{
+		XR_ENSURE(xrDestroySession(Session));
+	}
+}
+
 bool FOpenXRHMD::OnStereoStartup()
 {
 	FOpenXRHMD* Self = this;
@@ -955,17 +963,51 @@ bool FOpenXRHMD::OnStereoTeardown()
 {
 	if (Session != XR_NULL_HANDLE)
 	{
-		xrRequestExitSession(Session);
+		XrResult Result = xrRequestExitSession(Session);
+		if (Result == XR_ERROR_SESSION_NOT_RUNNING)
+		{
+			// Session was never running - most likely PIE without putting the headset on. 
+			CloseSession();
+		}
+		else
+		{
+			XR_ENSURE(Result);
+		}
 	}
 
 	return true;
 }
 
-FOpenXRHMD::~FOpenXRHMD()
+void FOpenXRHMD::CloseSession()
 {
-	if (Session)
+	if (Session != XR_NULL_HANDLE)
 	{
-		XR_ENSURE(xrDestroySession(Session));
+		// Clear up action spaces
+		for (auto& ActionSpace : ActionSpaces)
+		{
+			ActionSpace.DestroySpace();
+		}
+		ActionSpaces.Empty();
+
+		// Close the session now we're allowed to.
+		ENQUEUE_RENDER_COMMAND(OpenXRDestroySession)([this](FRHICommandListImmediate& RHICmdList)
+		{
+			if (bIsRunning)
+			{
+				XR_ENSURE(xrEndSession(Session));
+			}
+
+			XR_ENSURE(xrDestroySession(Session));
+
+			Session = XR_NULL_HANDLE;
+		});
+
+		FlushRenderingCommands();
+
+		bIsRendering = false;
+		bIsReady = false;
+		bIsRunning = false;
+		bRunRequested = false;
 	}
 }
 
@@ -1275,34 +1317,7 @@ bool FOpenXRHMD::OnStartGameFrame(FWorldContext& WorldContext)
 				FPlatformMisc::RequestExit(false);
 			}
 
-			if (Session != XR_NULL_HANDLE)
-			{
-				// Clear up action spaces
-				for (auto& ActionSpace : ActionSpaces)
-				{
-					ActionSpace.DestroySpace();
-				}
-
-				// Close the session now we're allowed to.
-				ENQUEUE_RENDER_COMMAND(OpenXRDestroySession)([this](FRHICommandListImmediate& RHICmdList)
-				{
-					if (bIsRunning)
-					{
-						XR_ENSURE(xrEndSession(Session));
-					}
-
-					XR_ENSURE(xrDestroySession(Session));
-
-					Session = XR_NULL_HANDLE;
-				});
-
-				FlushRenderingCommands();
-
-				bIsRendering = false;
-				bIsReady = false;
-				bIsRunning = false;
-				bRunRequested = false;
-			}
+			CloseSession();
 
 			break;
 		}
diff --git a/Engine/Plugins/Runtime/OpenXR/Source/OpenXRHMD/Private/OpenXRHMD.h b/Engine/Plugins/Runtime/OpenXR/Source/OpenXRHMD/Private/OpenXRHMD.h
--- a/Engine/Plugins/Runtime/OpenXR/Source/OpenXRHMD/Private/OpenXRHMD.h
+++ b/Engine/Plugins/Runtime/OpenXR/Source/OpenXRHMD/Private/OpenXRHMD.h
@@ -107,6 +107,7 @@ protected:
 	bool OnStereoStartup();
 	bool OnStereoTeardown();
 	bool ReadNextEvent(XrEventDataBuffer* buffer);
+	void CloseSession();
 
 public:
 	/** IHeadMountedDisplay interface */
