Remove code that emptied the ActionSpaces array on session end (this array needs to be kept around between sessions).

diff --git a/Engine/Plugins/Runtime/OpenXR/Source/OpenXRHMD/Private/OpenXRHMD.cpp b/Engine/Plugins/Runtime/OpenXR/Source/OpenXRHMD/Private/OpenXRHMD.cpp
--- a/Engine/Plugins/Runtime/OpenXR/Source/OpenXRHMD/Private/OpenXRHMD.cpp
+++ b/Engine/Plugins/Runtime/OpenXR/Source/OpenXRHMD/Private/OpenXRHMD.cpp
@@ -987,7 +987,6 @@ void FOpenXRHMD::CloseSession()
 		{
 			ActionSpace.DestroySpace();
 		}
-		ActionSpaces.Empty();
 
 		// Close the session now we're allowed to.
 		ENQUEUE_RENDER_COMMAND(OpenXRDestroySession)([this](FRHICommandListImmediate& RHICmdList)
