OpenXR Input refresh, first pass (unshelved reviewed work by Jules).

diff --git a/Engine/Config/BaseInput.ini b/Engine/Config/BaseInput.ini
--- a/Engine/Config/BaseInput.ini
+++ b/Engine/Config/BaseInput.ini
@@ -241,4 +241,84 @@ AutoCompleteCommandColor=(B=185,G=109,R=144,A=255)
 AutoCompleteCVarColor=(B=86,G=158,R=86,A=255)
 AutoCompleteFadedColor=(R=100,G=100,B=100,A=255)
 
-
+[/Script/OpenXRInput.XRInputSettings]
+SimpleController=(Supported=False,ControllerPose=GripPose)
+DaydreamController=(Supported=False,ControllerPose=GripPose)
+ViveController=(Supported=True,ControllerPose=GripPose)
++ViveBindings=(Path="/user/hand/left/input/menu/click",Key=MotionController_Left_Shoulder)
++ViveBindings=(Path="/user/hand/right/input/menu/click",Key=MotionController_Right_Shoulder)
++ViveBindings=(Path="/user/hand/left/input/squeeze",Key=MotionController_Left_Grip1)
++ViveBindings=(Path="/user/hand/right/input/squeeze",Key=MotionController_Right_Grip1)
++ViveBindings=(Path="/user/hand/left/input/squeeze",Key=MotionController_Left_Grip1Axis)
++ViveBindings=(Path="/user/hand/right/input/squeeze",Key=MotionController_Right_Grip1Axis)
++ViveBindings=(Path="/user/hand/left/input/trigger",Key=MotionController_Left_Trigger)
++ViveBindings=(Path="/user/hand/right/input/trigger",Key=MotionController_Right_Trigger)
++ViveBindings=(Path="/user/hand/left/input/trigger",Key=MotionController_Left_TriggerAxis)
++ViveBindings=(Path="/user/hand/right/input/trigger",Key=MotionController_Right_TriggerAxis)
++ViveBindings=(Path="/user/hand/left/input/trackpad/x",Key=MotionController_Left_Thumbstick_X)
++ViveBindings=(Path="/user/hand/right/input/trackpad/x",Key=MotionController_Right_Thumbstick_X)
++ViveBindings=(Path="/user/hand/left/input/trackpad/y",Key=MotionController_Left_Thumbstick_Y)
++ViveBindings=(Path="/user/hand/right/input/trackpad/y",Key=MotionController_Right_Thumbstick_Y)
++ViveBindings=(Path="/user/hand/left/input/trackpad/click",Key=MotionController_Left_Thumbstick)
++ViveBindings=(Path="/user/hand/right/input/trackpad/click",Key=MotionController_Right_Thumbstick)
+ViveProHeadset=(Supported=False,ControllerPose=GripPose)
+MixedRealityController=(Supported=True,ControllerPose=GripPose)
++MixedRealityBindings=(Path="/user/hand/left/input/menu/click",Key=MotionController_Left_Shoulder)
++MixedRealityBindings=(Path="/user/hand/right/input/menu/click",Key=MotionController_Right_Shoulder)
++MixedRealityBindings=(Path="/user/hand/left/input/squeeze",Key=MotionController_Left_Grip1)
++MixedRealityBindings=(Path="/user/hand/right/input/squeeze",Key=MotionController_Right_Grip1)
++MixedRealityBindings=(Path="/user/hand/left/input/squeeze",Key=MotionController_Left_Grip1Axis)
++MixedRealityBindings=(Path="/user/hand/right/input/squeeze",Key=MotionController_Right_Grip1Axis)
++MixedRealityBindings=(Path="/user/hand/left/input/trigger",Key=MotionController_Left_Trigger)
++MixedRealityBindings=(Path="/user/hand/right/input/trigger",Key=MotionController_Right_Trigger)
++MixedRealityBindings=(Path="/user/hand/left/input/trigger",Key=MotionController_Left_TriggerAxis)
++MixedRealityBindings=(Path="/user/hand/right/input/trigger",Key=MotionController_Right_TriggerAxis)
++MixedRealityBindings=(Path="/user/hand/left/input/thumbstick/x",Key=MotionController_Left_Thumbstick_X)
++MixedRealityBindings=(Path="/user/hand/right/input/thumbstick/x",Key=MotionController_Right_Thumbstick_X)
++MixedRealityBindings=(Path="/user/hand/left/input/thumbstick/y",Key=MotionController_Left_Thumbstick_Y)
++MixedRealityBindings=(Path="/user/hand/right/input/thumbstick/y",Key=MotionController_Right_Thumbstick_Y)
++MixedRealityBindings=(Path="/user/hand/left/input/thumbstick/click",Key=MotionController_Left_Thumbstick)
++MixedRealityBindings=(Path="/user/hand/right/input/thumbstick/click",Key=MotionController_Right_Thumbstick)
+OculusGoController=(Supported=False,ControllerPose=GripPose)
+OculusTouchController=(Supported=True,ControllerPose=GripPose)
++OculusTouchBindings=(Path="/user/hand/left/input/menu/click",Key=MotionController_Left_Shoulder)
++OculusTouchBindings=(Path="/user/hand/right/input/menu/click",Key=MotionController_Right_Shoulder)
++OculusTouchBindings=(Path="/user/hand/left/input/squeeze",Key=MotionController_Left_Grip1)
++OculusTouchBindings=(Path="/user/hand/right/input/squeeze",Key=MotionController_Right_Grip1)
++OculusTouchBindings=(Path="/user/hand/left/input/squeeze",Key=MotionController_Left_Grip1Axis)
++OculusTouchBindings=(Path="/user/hand/right/input/squeeze",Key=MotionController_Right_Grip1Axis)
++OculusTouchBindings=(Path="/user/hand/left/input/trigger",Key=MotionController_Left_Trigger)
++OculusTouchBindings=(Path="/user/hand/right/input/trigger",Key=MotionController_Right_Trigger)
++OculusTouchBindings=(Path="/user/hand/left/input/trigger",Key=MotionController_Left_TriggerAxis)
++OculusTouchBindings=(Path="/user/hand/right/input/trigger",Key=MotionController_Right_TriggerAxis)
++OculusTouchBindings=(Path="/user/hand/left/input/thumbstick/x",Key=MotionController_Left_Thumbstick_X)
++OculusTouchBindings=(Path="/user/hand/right/input/thumbstick/x",Key=MotionController_Right_Thumbstick_X)
++OculusTouchBindings=(Path="/user/hand/left/input/thumbstick/y",Key=MotionController_Left_Thumbstick_Y)
++OculusTouchBindings=(Path="/user/hand/right/input/thumbstick/y",Key=MotionController_Right_Thumbstick_Y)
++OculusTouchBindings=(Path="/user/hand/left/input/thumbstick/click",Key=MotionController_Left_Thumbstick)
++OculusTouchBindings=(Path="/user/hand/right/input/thumbstick/click",Key=MotionController_Right_Thumbstick)
++OculusTouchBindings=(Path="/user/hand/left/input/x/click",Key=MotionController_Left_FaceButton1)
++OculusTouchBindings=(Path="/user/hand/right/input/a/click",Key=MotionController_Right_FaceButton1)
++OculusTouchBindings=(Path="/user/hand/left/input/y/click",Key=MotionController_Left_FaceButton2)
++OculusTouchBindings=(Path="/user/hand/right/input/b/click",Key=MotionController_Right_FaceButton2)
+ValveIndexController=(Supported=True,ControllerPose=GripPose)
++ValveIndexBindings=(Path="/user/hand/left/input/menu/click",Key=MotionController_Left_Shoulder)
++ValveIndexBindings=(Path="/user/hand/right/input/menu/click",Key=MotionController_Right_Shoulder)
++ValveIndexBindings=(Path="/user/hand/left/input/squeeze",Key=MotionController_Left_Grip1)
++ValveIndexBindings=(Path="/user/hand/right/input/squeeze",Key=MotionController_Right_Grip1)
++ValveIndexBindings=(Path="/user/hand/left/input/squeeze",Key=MotionController_Left_Grip1Axis)
++ValveIndexBindings=(Path="/user/hand/right/input/squeeze",Key=MotionController_Right_Grip1Axis)
++ValveIndexBindings=(Path="/user/hand/left/input/trigger",Key=MotionController_Left_Trigger)
++ValveIndexBindings=(Path="/user/hand/right/input/trigger",Key=MotionController_Right_Trigger)
++ValveIndexBindings=(Path="/user/hand/left/input/trigger",Key=MotionController_Left_TriggerAxis)
++ValveIndexBindings=(Path="/user/hand/right/input/trigger",Key=MotionController_Right_TriggerAxis)
++ValveIndexBindings=(Path="/user/hand/left/input/thumbstick/x",Key=MotionController_Left_Thumbstick_X)
++ValveIndexBindings=(Path="/user/hand/right/input/thumbstick/x",Key=MotionController_Right_Thumbstick_X)
++ValveIndexBindings=(Path="/user/hand/left/input/thumbstick/y",Key=MotionController_Left_Thumbstick_Y)
++ValveIndexBindings=(Path="/user/hand/right/input/thumbstick/y",Key=MotionController_Right_Thumbstick_Y)
++ValveIndexBindings=(Path="/user/hand/left/input/thumbstick/click",Key=MotionController_Left_Thumbstick)
++ValveIndexBindings=(Path="/user/hand/right/input/thumbstick/click",Key=MotionController_Right_Thumbstick)
++ValveIndexBindings=(Path="/user/hand/left/input/a/click",Key=MotionController_Left_FaceButton1)
++ValveIndexBindings=(Path="/user/hand/right/input/a/click",Key=MotionController_Right_FaceButton1)
++ValveIndexBindings=(Path="/user/hand/left/input/b/click",Key=MotionController_Left_FaceButton2)
++ValveIndexBindings=(Path="/user/hand/right/input/b/click",Key=MotionController_Right_FaceButton2)
diff --git a/Engine/Plugins/Runtime/OpenXR/Source/OpenXRInput/Classes/XRInputSettings.h b/Engine/Plugins/Runtime/OpenXR/Source/OpenXRInput/Classes/XRInputSettings.h
new file mode 100644
--- /dev/null
+++ b/Engine/Plugins/Runtime/OpenXR/Source/OpenXRInput/Classes/XRInputSettings.h
@@ -0,0 +1,128 @@
+// Copyright 1998-2019 Epic Games, Inc. All Rights Reserved.
+
+#pragma once
+#include <openxr/openxr.h>
+#include "CoreMinimal.h"
+#include "InputCoreTypes.h"
+#include "XRInputSettings.generated.h"
+
+UENUM()
+enum EXRControllerPose
+{
+	GripPose,
+	AimPose
+};
+
+USTRUCT()
+struct FXRSuggestedBinding
+{
+	GENERATED_BODY()
+
+	UPROPERTY(config, EditAnywhere, Category = Input)
+	FString Path;
+
+	UPROPERTY(config, EditAnywhere, Category = Input)
+	FKey Key;
+};
+
+USTRUCT()
+struct FXRInteractionProfileSettings
+{
+	GENERATED_BODY()
+
+	bool HasControllers;
+	bool HasHaptics;
+
+	/** Indicates whether this profile is supported by the application. */
+	UPROPERTY(config, EditAnywhere, Category = Input)
+	bool Supported;
+
+	/** Indicates whether the controller is held in a grip or aim pose. */
+	UPROPERTY(config, EditAnywhere, Category = Input)
+	TEnumAsByte<EXRControllerPose> ControllerPose;
+};
+
+/**
+* Implements the settings for OpenXR input.
+*/
+UCLASS(config = Input, defaultconfig)
+class OPENXRINPUT_API UXRInputSettings : public UObject
+{
+public:
+	GENERATED_BODY()
+
+	UPROPERTY(config, EditAnywhere, Category = "Interaction Profiles")
+	FXRInteractionProfileSettings SimpleController;
+	
+	UPROPERTY(config, EditAnywhere, Category = "Suggested Bindings")
+	TArray<FXRSuggestedBinding> SimpleBindings;
+
+	UPROPERTY(config, EditAnywhere, Category = "Interaction Profiles")
+	FXRInteractionProfileSettings DaydreamController;
+	
+	UPROPERTY(config, EditAnywhere, Category = "Suggested Bindings")
+	TArray<FXRSuggestedBinding> DaydreamBindings;
+
+	UPROPERTY(config, EditAnywhere, Category = "Interaction Profiles")
+	FXRInteractionProfileSettings ViveController;
+	
+	UPROPERTY(config, EditAnywhere, Category = "Suggested Bindings")
+	TArray<FXRSuggestedBinding> ViveBindings;
+
+	UPROPERTY(config, EditAnywhere, Category = "Interaction Profiles")
+	FXRInteractionProfileSettings ViveProHeadset;
+	
+	UPROPERTY(config, EditAnywhere, Category = "Suggested Bindings")
+	TArray<FXRSuggestedBinding> ViveProBindings;
+
+	UPROPERTY(config, EditAnywhere, Category = "Interaction Profiles")
+	FXRInteractionProfileSettings MixedRealityController;
+	
+	UPROPERTY(config, EditAnywhere, Category = "Suggested Bindings")
+	TArray<FXRSuggestedBinding> MixedRealityBindings;
+
+	UPROPERTY(config, EditAnywhere, Category = "Interaction Profiles")
+	FXRInteractionProfileSettings OculusGoController;
+	
+	UPROPERTY(config, EditAnywhere, Category = "Suggested Bindings")
+	TArray<FXRSuggestedBinding> OculusGoBindings;
+
+	UPROPERTY(config, EditAnywhere, Category = "Interaction Profiles")
+	FXRInteractionProfileSettings OculusTouchController;
+	
+	UPROPERTY(config, EditAnywhere, Category = "Suggested Bindings")
+	TArray<FXRSuggestedBinding> OculusTouchBindings;
+
+	UPROPERTY(config, EditAnywhere, Category = "Interaction Profiles")
+	FXRInteractionProfileSettings ValveIndexController;
+	
+	UPROPERTY(config, EditAnywhere, Category = "Suggested Bindings")
+	TArray<FXRSuggestedBinding> ValveIndexBindings;
+
+#if WITH_EDITOR
+	virtual void PostEditChangeChainProperty(struct FPropertyChangedChainEvent& PropertyChangedEvent) override;
+#endif
+
+	static FSimpleMulticastDelegate OnSuggestedBindingsChanged;
+
+	UXRInputSettings()
+	{
+		SimpleController.HasControllers = true;
+		DaydreamController.HasControllers = true;
+		ViveController.HasControllers = true;
+		ViveProHeadset.HasControllers = false;
+		MixedRealityController.HasControllers = true;
+		OculusGoController.HasControllers = true;
+		OculusTouchController.HasControllers = true;
+		ValveIndexController.HasControllers = true;
+
+		SimpleController.HasHaptics = false;
+		DaydreamController.HasHaptics = false;
+		ViveController.HasHaptics = true;
+		ViveProHeadset.HasHaptics = false;
+		MixedRealityController.HasHaptics = true;
+		OculusGoController.HasHaptics = false;
+		OculusTouchController.HasHaptics = true;
+		ValveIndexController.HasHaptics = true;
+	}
+};
diff --git a/Engine/Plugins/Runtime/OpenXR/Source/OpenXRInput/OpenXRInput.Build.cs b/Engine/Plugins/Runtime/OpenXR/Source/OpenXRInput/OpenXRInput.Build.cs
--- a/Engine/Plugins/Runtime/OpenXR/Source/OpenXRInput/OpenXRInput.Build.cs
+++ b/Engine/Plugins/Runtime/OpenXR/Source/OpenXRInput/OpenXRInput.Build.cs
@@ -34,6 +34,11 @@ namespace UnrealBuildTool.Rules
                 );
 
             AddEngineThirdPartyPrivateStaticDependencies(Target, "OpenXR");
+
+            if (Target.bBuildEditor == true)
+            {
+                PrivateDependencyModuleNames.Add("UnrealEd");
+            }
         }
     }
 }
diff --git a/Engine/Plugins/Runtime/OpenXR/Source/OpenXRInput/Private/OpenXRInput.cpp b/Engine/Plugins/Runtime/OpenXR/Source/OpenXRInput/Private/OpenXRInput.cpp
--- a/Engine/Plugins/Runtime/OpenXR/Source/OpenXRInput/Private/OpenXRInput.cpp
+++ b/Engine/Plugins/Runtime/OpenXR/Source/OpenXRInput/Private/OpenXRInput.cpp
@@ -3,21 +3,29 @@
 #include "OpenXRInput.h"
 #include "OpenXRHMD.h"
 #include "OpenXRHMDPrivate.h"
+#include "XRInputSettings.h"
 #include "UObject/UObjectIterator.h"
 #include "GameFramework/InputSettings.h"
 
+#if WITH_EDITOR
+#include "Editor/EditorEngine.h"
+#include "ISettingsModule.h"
+#include "ISettingsSection.h"
+#include "Editor.h"
+#endif
+
 #include <openxr/openxr.h>
 
-// Hack to prefer emitting MotionController keys for action events
-static bool MatchKeyNamePrefix(const FKey& Key, const TCHAR* Prefix)
-{
-	return Key.GetFName().ToString().StartsWith(Prefix);
-};
+#define LOCTEXT_NAMESPACE "OpenXR"
+
+FSimpleMulticastDelegate UXRInputSettings::OnSuggestedBindingsChanged;
 
-static bool MatchKeyNameSuffix(const FKey& Key, const TCHAR* Suffix)
+void UXRInputSettings::PostEditChangeChainProperty(struct FPropertyChangedChainEvent& PropertyChangedEvent)
 {
-	return Key.GetFName().ToString().EndsWith(Suffix);
-};
+	Super::PostEditChangeChainProperty(PropertyChangedEvent);
+
+	OnSuggestedBindingsChanged.Broadcast();
+}
 
 FORCEINLINE XrPath GetPath(XrInstance Instance, const char* PathString)
 {
@@ -27,6 +35,12 @@ FORCEINLINE XrPath GetPath(XrInstance Instance, const char* PathString)
 	return Path;
 }
 
+FORCEINLINE XrPath GetPath(XrInstance Instance, FString PathString)
+{
+	FTCHARToUTF8 EngineNameConverter(*PathString);
+	return GetPath(Instance, EngineNameConverter.Get());
+}
+
 FORCEINLINE void FilterActionName(const char* InActionName, char* OutActionName)
 {
 	// Ensure the action name is a well-formed path
@@ -75,15 +89,44 @@ void FOpenXRInputPlugin::StartupModule()
 	FOpenXRHMD* OpenXRHMD = GetOpenXRHMD();
 	if (OpenXRHMD)
 	{
-		InputDevice = TSharedPtr<FOpenXRInput>(new FOpenXRInputPlugin::FOpenXRInput(OpenXRHMD));
+		InputDevice = MakeShared<FOpenXRInput>(OpenXRHMD);
+	}
+
+#if WITH_EDITOR
+	ISettingsModule* SettingsModule = FModuleManager::GetModulePtr<ISettingsModule>("Settings");
+
+	// While this should usually be true, it's not guaranteed that the settings module will be loaded in the editor.
+	// UBT allows setting bBuildDeveloperTools to false while bBuildEditor can be true.
+	// The former option indirectly controls loading of the "Settings" module.
+	if (SettingsModule)
+	{
+		SettingsModule->RegisterSettings("Project", "Plugins", "OpenXR Input",
+			LOCTEXT("XRInputSettingsName", "OpenXR Input"),
+			LOCTEXT("XRInputSettingsDescription", "Configure input for OpenXR."),
+			GetMutableDefault<UXRInputSettings>()
+		);
+	}
+#endif
+}
+
+void FOpenXRInputPlugin::ShutdownModule()
+{
+	IOpenXRInputPlugin::ShutdownModule();
+
+#if WITH_EDITOR
+	ISettingsModule* SettingsModule = FModuleManager::GetModulePtr<ISettingsModule>("Settings");
+	if (SettingsModule)
+	{
+		SettingsModule->UnregisterSettings("Project", "Plugins", "OpenXR Input");
 	}
+#endif
 }
 
 FOpenXRInputPlugin::FOpenXRAction::FOpenXRAction(XrActionSet InActionSet, XrActionType InActionType, const FName& InName)
 	: Set(InActionSet)
 	, Type(InActionType)
 	, Name(InName)
-	, ActionKey()
+	, Keys()
 	, Handle(XR_NULL_HANDLE)
 {
 	char ActionName[NAME_SIZE];
@@ -100,16 +143,22 @@ FOpenXRInputPlugin::FOpenXRAction::FOpenXRAction(XrActionSet InActionSet, XrActi
 	XR_ENSURE(xrCreateAction(Set, &Info, &Handle));
 }
 
-FOpenXRInputPlugin::FOpenXRAction::FOpenXRAction(XrActionSet InSet, const FInputActionKeyMapping& InActionKey)
-	: FOpenXRAction(InSet, XR_ACTION_TYPE_BOOLEAN_INPUT, InActionKey.ActionName)
+FOpenXRInputPlugin::FOpenXRAction::FOpenXRAction(XrActionSet InSet, const TArray<FInputActionKeyMapping>& InActionKeys)
+	: FOpenXRAction(InSet, XR_ACTION_TYPE_BOOLEAN_INPUT, InActionKeys[0].ActionName)
 {
-	ActionKey = InActionKey.Key.GetFName();
+	for (auto Mapping : InActionKeys)
+	{
+		Keys.Add(Mapping.Key);
+	}
 }
 
-FOpenXRInputPlugin::FOpenXRAction::FOpenXRAction(XrActionSet InSet, const FInputAxisKeyMapping& InAxisKey)
-	: FOpenXRAction(InSet, XR_ACTION_TYPE_FLOAT_INPUT, InAxisKey.AxisName)
+FOpenXRInputPlugin::FOpenXRAction::FOpenXRAction(XrActionSet InSet, const TArray<FInputAxisKeyMapping>& InAxisKeys)
+	: FOpenXRAction(InSet, XR_ACTION_TYPE_FLOAT_INPUT, InAxisKeys[0].AxisName)
 {
-	ActionKey = InAxisKey.Key.GetFName();
+	for (auto Mapping : InAxisKeys)
+	{
+		Keys.Add(Mapping.Key);
+	}
 }
 
 FOpenXRInputPlugin::FOpenXRController::FOpenXRController(FOpenXRHMD* HMD, XrActionSet InActionSet, const char* InName)
@@ -152,8 +201,22 @@ FOpenXRInputPlugin::FOpenXRInput::FOpenXRInput(FOpenXRHMD* HMD)
 	IModularFeatures::Get().RegisterModularFeature(GetModularFeatureName(), this);
 	check(OpenXRHMD);
 
+	InitActions();
+}
+
+FOpenXRInputPlugin::FOpenXRInput::~FOpenXRInput()
+{
+	DestroyActions();
+}
+
+void FOpenXRInputPlugin::FOpenXRInput::InitActions()
+{
+	if (bActionsBound)
+		return;
+
 	XrInstance Instance = OpenXRHMD->GetInstance();
 	check(Instance);
+	DestroyActions();
 
 	XrActionSet ActionSet = XR_NULL_HANDLE;
 	XrActionSetCreateInfo SetInfo;
@@ -163,40 +226,7 @@ FOpenXRInputPlugin::FOpenXRInput::FOpenXRInput(FOpenXRHMD* HMD)
 	FCStringAnsi::Strcpy(SetInfo.localizedActionSetName, XR_MAX_ACTION_SET_NAME_SIZE, "Unreal Engine 4");
 	XR_ENSURE(xrCreateActionSet(Instance, &SetInfo, &ActionSet));
 
-	// Controller poses
-	TArray<XrActionSuggestedBinding> Bindings;
-	Controllers.Add(EControllerHand::Left, FOpenXRController(OpenXRHMD, ActionSet, "Left Controller"));
-	Controllers.Add(EControllerHand::Right, FOpenXRController(OpenXRHMD, ActionSet, "Right Controller"));
-
-	Bindings.Add(XrActionSuggestedBinding{ Controllers[EControllerHand::Left].Action, GetPath(Instance, "/user/hand/left/input/grip/pose") });
-	Bindings.Add(XrActionSuggestedBinding{ Controllers[EControllerHand::Right].Action, GetPath(Instance, "/user/hand/right/input/grip/pose") });
-
-	Bindings.Add(XrActionSuggestedBinding{ Controllers[EControllerHand::Left].VibrationAction, GetPath(Instance, "/user/hand/left/output/haptic") });
-	Bindings.Add(XrActionSuggestedBinding{ Controllers[EControllerHand::Right].VibrationAction, GetPath(Instance, "/user/hand/right/output/haptic") });
-
-	InteractionMappings.Add(FGamepadKeyNames::MotionController_Left_Shoulder, GetPath(Instance, "/user/hand/left/input/menu/click"));
-	InteractionMappings.Add(FGamepadKeyNames::MotionController_Left_Trigger, GetPath(Instance, "/user/hand/left/input/trigger"));
-	InteractionMappings.Add(FGamepadKeyNames::MotionController_Left_TriggerAxis, GetPath(Instance, "/user/hand/left/input/trigger"));
-	InteractionMappings.Add(FGamepadKeyNames::MotionController_Left_Grip1, GetPath(Instance, "/user/hand/left/input/squeeze"));
-	InteractionMappings.Add(FGamepadKeyNames::MotionController_Left_Grip1Axis, GetPath(Instance, "/user/hand/left/input/squeeze"));
-	InteractionMappings.Add(FGamepadKeyNames::MotionController_Left_Thumbstick_X, GetPath(Instance, "/user/hand/left/input/thumbstick/x"));
-	InteractionMappings.Add(FGamepadKeyNames::MotionController_Left_Thumbstick_Y, GetPath(Instance, "/user/hand/left/input/thumbstick/y"));
-	InteractionMappings.Add(FGamepadKeyNames::MotionController_Left_Thumbstick, GetPath(Instance, "/user/hand/left/input/thumbstick/click"));
-	//InteractionMappings.Add(FGamepadKeyNames::MotionController_Left_FaceButton1, GetPath(Instance, "/user/hand/left/input/x/click"));
-	//InteractionMappings.Add(FGamepadKeyNames::MotionController_Left_FaceButton2, GetPath(Instance, "/user/hand/left/input/y/click"));
-
-	InteractionMappings.Add(FGamepadKeyNames::MotionController_Right_Shoulder, GetPath(Instance, "/user/hand/right/input/menu/click"));
-	InteractionMappings.Add(FGamepadKeyNames::MotionController_Right_Trigger, GetPath(Instance, "/user/hand/right/input/trigger"));
-	InteractionMappings.Add(FGamepadKeyNames::MotionController_Right_TriggerAxis, GetPath(Instance, "/user/hand/right/input/trigger"));
-	InteractionMappings.Add(FGamepadKeyNames::MotionController_Right_Grip1, GetPath(Instance, "/user/hand/right/input/squeeze"));
-	InteractionMappings.Add(FGamepadKeyNames::MotionController_Right_Grip1Axis, GetPath(Instance, "/user/hand/right/input/squeeze"));
-	InteractionMappings.Add(FGamepadKeyNames::MotionController_Right_Thumbstick_X, GetPath(Instance, "/user/hand/right/input/thumbstick/x"));
-	InteractionMappings.Add(FGamepadKeyNames::MotionController_Right_Thumbstick_Y, GetPath(Instance, "/user/hand/right/input/thumbstick/y"));
-	InteractionMappings.Add(FGamepadKeyNames::MotionController_Right_Thumbstick, GetPath(Instance, "/user/hand/right/input/thumbstick/click"));
-	//InteractionMappings.Add(FGamepadKeyNames::MotionController_Right_FaceButton1, GetPath(Instance, "/user/hand/right/input/a/click"));
-	//InteractionMappings.Add(FGamepadKeyNames::MotionController_Right_FaceButton2, GetPath(Instance, "/user/hand/right/input/b/click"));
-
-	auto InputSettings = GetDefault<UInputSettings>();
+	const UInputSettings* InputSettings = GetDefault<UInputSettings>();
 	if (InputSettings != nullptr)
 	{
 		TArray<FName> ActionNames;
@@ -205,7 +235,7 @@ FOpenXRInputPlugin::FOpenXRInput::FOpenXRInput(FOpenXRHMD* HMD)
 		{
 			TArray<FInputActionKeyMapping> Mappings;
 			InputSettings->GetActionMappingByName(ActionName, Mappings);
-			AddAction(ActionSet, Mappings, Bindings);
+			Actions.Emplace(ActionSet, Mappings);
 		}
 
 		TArray<FName> AxisNames;
@@ -214,35 +244,24 @@ FOpenXRInputPlugin::FOpenXRInput::FOpenXRInput(FOpenXRHMD* HMD)
 		{
 			TArray<FInputAxisKeyMapping> Mappings;
 			InputSettings->GetAxisMappingByName(AxisName, Mappings);
-			AddAction(ActionSet, Mappings, Bindings);
-		}
-
-		// Open console
-		{
-			const FKey* ConsoleKey = InputSettings->ConsoleKeys.FindByPredicate([](FKey& Key) { return Key.IsValid(); });
-			if (ConsoleKey != nullptr)
-			{
-				Actions.Add(FOpenXRAction(ActionSet, FName(TEXT("Open Console")), ConsoleKey->GetFName()));
-			}
+			Actions.Emplace(ActionSet, Mappings);
 		}
 	}
 
-	TArray<XrPath> Profiles;
-	//Profiles.Add(GetPath(Instance, "/interaction_profiles/khr/simple_controller"));
-	Profiles.Add(GetPath(Instance, "/interaction_profiles/microsoft/motion_controller"));
-	Profiles.Add(GetPath(Instance, "/interaction_profiles/oculus/touch_controller"));
-	Profiles.Add(GetPath(Instance, "/interaction_profiles/valve/index_controller"));
+	const UXRInputSettings* XRInputSettings = GetDefault<UXRInputSettings>();
 
-	for (XrPath Profile : Profiles)
-	{
-		XrInteractionProfileSuggestedBinding InteractionProfile;
-		InteractionProfile.type = XR_TYPE_INTERACTION_PROFILE_SUGGESTED_BINDING;
-		InteractionProfile.next = nullptr;
-		InteractionProfile.interactionProfile = Profile;
-		InteractionProfile.countSuggestedBindings = Bindings.Num();
-		InteractionProfile.suggestedBindings = Bindings.GetData();
-		XR_ENSURE(xrSuggestInteractionProfileBindings(Instance, &InteractionProfile));
-	}
+	// Controller poses
+	Controllers.Add(EControllerHand::Left, FOpenXRController(OpenXRHMD, ActionSet, "Left Controller"));
+	Controllers.Add(EControllerHand::Right, FOpenXRController(OpenXRHMD, ActionSet, "Right Controller"));
+
+	SuggestedBindings(Instance, "/interaction_profiles/khr/simple_controller", XRInputSettings->SimpleController, XRInputSettings->SimpleBindings);
+	SuggestedBindings(Instance, "/interaction_profiles/google/daydream_controller", XRInputSettings->DaydreamController, XRInputSettings->DaydreamBindings);
+	SuggestedBindings(Instance, "/interaction_profiles/htc/vive_controller", XRInputSettings->ViveController, XRInputSettings->ViveBindings);
+	SuggestedBindings(Instance, "/interaction_profiles/htc/vive_pro", XRInputSettings->ViveProHeadset, XRInputSettings->ViveProBindings);
+	SuggestedBindings(Instance, "/interaction_profiles/microsoft/motion_controller", XRInputSettings->MixedRealityController, XRInputSettings->MixedRealityBindings);
+	SuggestedBindings(Instance, "/interaction_profiles/oculus/go_controller", XRInputSettings->OculusGoController, XRInputSettings->OculusGoBindings);
+	SuggestedBindings(Instance, "/interaction_profiles/oculus/touch_controller", XRInputSettings->OculusTouchController, XRInputSettings->OculusTouchBindings);
+	SuggestedBindings(Instance, "/interaction_profiles/valve/index_controller", XRInputSettings->ValveIndexController, XRInputSettings->ValveIndexBindings);
 
 	XrActiveActionSet ActiveSet;
 	ActiveSet.actionSet = ActionSet;
@@ -250,50 +269,87 @@ FOpenXRInputPlugin::FOpenXRInput::FOpenXRInput(FOpenXRHMD* HMD)
 	ActionSets.Add(ActiveSet);
 }
 
-FOpenXRInputPlugin::FOpenXRInput::~FOpenXRInput()
+void FOpenXRInputPlugin::FOpenXRInput::DestroyActions()
 {
 	for (XrActiveActionSet& ActionSet : ActionSets)
 	{
-		xrDestroyActionSet(ActionSet.actionSet);
+		XR_ENSURE(xrDestroyActionSet(ActionSet.actionSet));
 	}
+	ActionSets.Empty();
+	Actions.Empty();
+	Controllers.Empty();
 }
 
-FORCEINLINE FName GetName(const FInputActionKeyMapping& ActionKey) { return ActionKey.ActionName; }
-FORCEINLINE FName GetName(const FInputAxisKeyMapping& AxisKey) { return AxisKey.AxisName; }
-
-template<typename T>
-void FOpenXRInputPlugin::FOpenXRInput::AddAction(XrActionSet ActionSet, const TArray<T>& Mappings, TArray<XrActionSuggestedBinding>& OutSuggestedBindings)
+void FOpenXRInputPlugin::FOpenXRInput::SuggestedBindings(XrInstance Instance, const char* Profile, const FXRInteractionProfileSettings& Settings, const TArray<FXRSuggestedBinding>& SuggestedBindings)
 {
-	// Find all the motion controller keys so we can suggest bindings for them
-	TArray<T> KeyMappings = Mappings.FilterByPredicate([](const T& Mapping)
+	if (!Settings.Supported)
 	{
-		return MatchKeyNamePrefix(Mapping.Key, TEXT("MotionController"));
-	});
-
-	// We need at least one key to be able to trigger the action
-	// TODO: Refactor the input API so we can trigger actions directly
-	if (KeyMappings.Num() == 0)
-	{
-		KeyMappings = Mappings.FilterByPredicate([](const T& Mapping) { return Mapping.Key.IsValid(); });
+		return;
 	}
 
-	if (KeyMappings.Num() > 0)
+	TArray<XrActionSuggestedBinding> Bindings;
+	for (const auto& Suggestion : SuggestedBindings)
 	{
-		// Build the action based on the primary mapping
-		FOpenXRAction Action(ActionSet, KeyMappings[0]);
+		if (Suggestion.Path.IsEmpty())
+		{
+			continue;
+		}
 
-		// Add suggested bindings for every mapping
-		for (T InputKey : KeyMappings)
+		const FKey& Key = Suggestion.Key;
+		const FOpenXRAction* Action = Actions.FindByPredicate([Key](FOpenXRAction& Mapping)
 		{
-			XrPath* Path = InteractionMappings.Find(InputKey.Key.GetFName());
-			if (Path)
+			int32 Index;
+			if (Mapping.Keys.Find(Key, Index))
 			{
-				OutSuggestedBindings.Add(XrActionSuggestedBinding{ Action.Handle, *Path });
+				Mapping.Keys.Swap(0, Index);
+				return true;
 			}
+			return false;
+		});
+
+		if (Action)
+		{
+			XrActionSuggestedBinding Binding;
+			Binding.action = Action->Handle;
+			Binding.binding = GetPath(Instance, Suggestion.Path);
+			Bindings.Add(Binding);
 		}
+	}
 
-		Actions.Add(Action);
+	if (Settings.HasControllers)
+	{
+		XrActionSuggestedBinding LeftBinding, RightBinding;
+		LeftBinding.action = Controllers[EControllerHand::Left].Action;
+		RightBinding.action = Controllers[EControllerHand::Right].Action;
+
+		if (Settings.ControllerPose == AimPose)
+		{
+			LeftBinding.binding = GetPath(Instance, "/user/hand/left/input/aim/pose");
+			RightBinding.binding = GetPath(Instance, "/user/hand/right/input/aim/pose");
+		}
+		else
+		{
+			LeftBinding.binding = GetPath(Instance, "/user/hand/left/input/grip/pose");
+			RightBinding.binding = GetPath(Instance, "/user/hand/right/input/grip/pose");
+		}
+
+		Bindings.Add(LeftBinding);
+		Bindings.Add(RightBinding);
+
+		if (Settings.HasHaptics)
+		{
+			Bindings.Add(XrActionSuggestedBinding{ Controllers[EControllerHand::Left].VibrationAction, GetPath(Instance, "/user/hand/left/output/haptic") });
+			Bindings.Add(XrActionSuggestedBinding{ Controllers[EControllerHand::Right].VibrationAction, GetPath(Instance, "/user/hand/right/output/haptic") });
+		}
 	}
+
+	XrInteractionProfileSuggestedBinding InteractionProfile;
+	InteractionProfile.type = XR_TYPE_INTERACTION_PROFILE_SUGGESTED_BINDING;
+	InteractionProfile.next = nullptr;
+	InteractionProfile.interactionProfile = GetPath(Instance, Profile);
+	InteractionProfile.countSuggestedBindings = Bindings.Num();
+	InteractionProfile.suggestedBindings = Bindings.GetData();
+	XR_ENSURE(xrSuggestInteractionProfileBindings(Instance, &InteractionProfile));
 }
 
 void FOpenXRInputPlugin::FOpenXRInput::Tick(float DeltaTime)
@@ -365,11 +421,11 @@ void FOpenXRInputPlugin::FOpenXRInput::SendControllerEvents()
 			{
 				if (State.currentState)
 				{
-					MessageHandler->OnControllerButtonPressed(Action.ActionKey, 0, /*IsRepeat =*/false);
+					MessageHandler->OnControllerButtonPressed(Action.Keys[0].GetFName(), 0, /*IsRepeat =*/false);
 				}
 				else
 				{
-					MessageHandler->OnControllerButtonReleased(Action.ActionKey, 0, /*IsRepeat =*/false);
+					MessageHandler->OnControllerButtonReleased(Action.Keys[0].GetFName(), 0, /*IsRepeat =*/false);
 				}
 			}
 		}
@@ -382,7 +438,7 @@ void FOpenXRInputPlugin::FOpenXRInput::SendControllerEvents()
 			XrResult Result = xrGetActionStateFloat(Session, &GetInfo, &State);
 			if (Result >= XR_SUCCESS && State.changedSinceLastSync)
 			{
-				MessageHandler->OnControllerAnalog(Action.ActionKey, 0, State.currentState);
+				MessageHandler->OnControllerAnalog(Action.Keys[0].GetFName(), 0, State.currentState);
 			}
 		}
 		break;
@@ -396,6 +452,9 @@ void FOpenXRInputPlugin::FOpenXRInput::SendControllerEvents()
 void FOpenXRInputPlugin::FOpenXRInput::SetMessageHandler(const TSharedRef< FGenericApplicationMessageHandler >& InMessageHandler)
 {
 	MessageHandler = InMessageHandler;
+
+	UXRInputSettings::OnSuggestedBindingsChanged.AddSP(this, &FOpenXRInputPlugin::FOpenXRInput::InitActions);
+	FEditorDelegates::OnActionAxisMappingsChanged.AddSP(this, &FOpenXRInputPlugin::FOpenXRInput::InitActions);
 }
 
 bool FOpenXRInputPlugin::FOpenXRInput::Exec(UWorld* InWorld, const TCHAR* Cmd, FOutputDevice& Ar)
@@ -520,3 +579,5 @@ float FOpenXRInputPlugin::FOpenXRInput::GetHapticAmplitudeScale() const
 {
 	return 1.0f;
 }
+
+#undef LOCTEXT_NAMESPACE
diff --git a/Engine/Plugins/Runtime/OpenXR/Source/OpenXRInput/Private/OpenXRInput.h b/Engine/Plugins/Runtime/OpenXR/Source/OpenXRInput/Private/OpenXRInput.h
--- a/Engine/Plugins/Runtime/OpenXR/Source/OpenXRInput/Private/OpenXRInput.h
+++ b/Engine/Plugins/Runtime/OpenXR/Source/OpenXRInput/Private/OpenXRInput.h
@@ -1,10 +1,12 @@
 // Copyright 1998-2019 Epic Games, Inc. All Rights Reserved.
 
+#pragma once
 #include "GenericPlatform/IInputInterface.h"
 #include "XRMotionControllerBase.h"
 #include "IOpenXRInputPlugin.h"
 #include "IInputDevice.h"
 #include "IHapticDevice.h"
+#include "XRInputSettings.h"
 
 #include <openxr/openxr.h>
 
@@ -20,19 +22,13 @@ public:
 		XrActionSet		Set;
 		XrActionType	Type;
 		FName			Name;
-		FName			ActionKey;
+		TArray<FKey>	Keys;
 		XrAction		Handle;
 
 		FOpenXRAction(XrActionSet InActionSet, XrActionType InActionType, const FName& InName);
 
-		FOpenXRAction(XrActionSet InActionSet, const FName& InName, const FName& InActionKey)
-			: FOpenXRAction(InActionSet, XR_ACTION_TYPE_BOOLEAN_INPUT, InName)
-		{
-			ActionKey = InActionKey;
-		}
-
-		FOpenXRAction(XrActionSet InSet, const FInputActionKeyMapping& InActionKey);
-		FOpenXRAction(XrActionSet InSet, const FInputAxisKeyMapping& InAxisKey);
+		FOpenXRAction(XrActionSet InSet, const TArray<FInputActionKeyMapping>& InActionKeys);
+		FOpenXRAction(XrActionSet InSet, const TArray<FInputAxisKeyMapping>& InAxisKeys);
 	};
 
 	struct FOpenXRController
@@ -45,7 +41,11 @@ public:
 		FOpenXRController(FOpenXRHMD* HMD, XrActionSet InActionSet, const char* InName);
 	};
 
-	class FOpenXRInput : public IInputDevice, public FXRMotionControllerBase, public IHapticDevice
+	class FOpenXRInput :
+		public IInputDevice,
+		public FXRMotionControllerBase,
+		public IHapticDevice,
+		public TSharedFromThis<FOpenXRInput>
 	{
 	public:
 		FOpenXRInput(FOpenXRHMD* HMD);
@@ -79,12 +79,12 @@ public:
 		TArray<XrActiveActionSet> ActionSets;
 		TArray<FOpenXRAction> Actions;
 		TMap<EControllerHand, FOpenXRController> Controllers;
-		TMap<FName, XrPath> InteractionMappings;
 
 		bool bActionsBound;
 
-		template<typename T>
-		void AddAction(XrActionSet ActionSet, const TArray<T>& Mappings, TArray<XrActionSuggestedBinding>& OutSuggestedBindings);
+		void InitActions();
+		void DestroyActions();
+		void SuggestedBindings(XrInstance Instance, const char* Profile, const FXRInteractionProfileSettings& Settings, const TArray<FXRSuggestedBinding>& SuggestedBindings);
 
 		/** handler to send all messages to */
 		TSharedRef<FGenericApplicationMessageHandler> MessageHandler;
@@ -96,6 +96,7 @@ public:
 	FOpenXRHMD* GetOpenXRHMD() const;
 
 	virtual void StartupModule() override;
+	virtual void ShutdownModule() override;
 	virtual TSharedPtr< class IInputDevice > CreateInputDevice(const TSharedRef< FGenericApplicationMessageHandler >& InMessageHandler) override;
 
 private:
