Fix build errors in OpenXR plugin introduced with new input code. Some editor-only code was wrapped with #if WITH_EDITOR in the header (declaration), but not in the implementation file (implementation), which was causing errors in non-editor builds.

diff --git a/Engine/Plugins/Runtime/OpenXR/Source/OpenXRInput/Private/OpenXRInput.cpp b/Engine/Plugins/Runtime/OpenXR/Source/OpenXRInput/Private/OpenXRInput.cpp
--- a/Engine/Plugins/Runtime/OpenXR/Source/OpenXRInput/Private/OpenXRInput.cpp
+++ b/Engine/Plugins/Runtime/OpenXR/Source/OpenXRInput/Private/OpenXRInput.cpp
@@ -20,12 +20,14 @@
 
 FSimpleMulticastDelegate UXRInputSettings::OnSuggestedBindingsChanged;
 
+#if WITH_EDITOR
 void UXRInputSettings::PostEditChangeChainProperty(struct FPropertyChangedChainEvent& PropertyChangedEvent)
 {
 	Super::PostEditChangeChainProperty(PropertyChangedEvent);
 
 	OnSuggestedBindingsChanged.Broadcast();
 }
+#endif
 
 FORCEINLINE XrPath GetPath(XrInstance Instance, const char* PathString)
 {
@@ -454,7 +456,9 @@ void FOpenXRInputPlugin::FOpenXRInput::SetMessageHandler(const TSharedRef< FGene
 	MessageHandler = InMessageHandler;
 
 	UXRInputSettings::OnSuggestedBindingsChanged.AddSP(this, &FOpenXRInputPlugin::FOpenXRInput::InitActions);
+#if WITH_EDITOR
 	FEditorDelegates::OnActionAxisMappingsChanged.AddSP(this, &FOpenXRInputPlugin::FOpenXRInput::InitActions);
+#endif
 }
 
 bool FOpenXRInputPlugin::FOpenXRInput::Exec(UWorld* InWorld, const TCHAR* Cmd, FOutputDevice& Ar)
