Implement hidden-area mesh support for OpenXR plugin.

diff --git a/Engine/Plugins/Runtime/OpenXR/Source/OpenXRHMD/Private/OpenXRHMD.cpp b/Engine/Plugins/Runtime/OpenXR/Source/OpenXRHMD/Private/OpenXRHMD.cpp
--- a/Engine/Plugins/Runtime/OpenXR/Source/OpenXRHMD/Private/OpenXRHMD.cpp
+++ b/Engine/Plugins/Runtime/OpenXR/Source/OpenXRHMD/Private/OpenXRHMD.cpp
@@ -150,8 +150,8 @@ TSharedPtr< class IXRTrackingSystem, ESPMode::ThreadSafe > FOpenXRHMDPlugin::Cre
 		}
 	}
 
-	auto OpenXRHMD = FSceneViewExtensions::NewExtension<FOpenXRHMD>(Instance, System, RenderBridge, HasExtension(XR_KHR_COMPOSITION_LAYER_DEPTH_EXTENSION_NAME));
-	if( OpenXRHMD->IsInitialized() )
+	auto OpenXRHMD = FSceneViewExtensions::NewExtension<FOpenXRHMD>(Instance, System, RenderBridge, AvailableExtensions);
+	if (OpenXRHMD->IsInitialized())
 	{
 		return OpenXRHMD;
 	}
@@ -362,6 +362,11 @@ bool FOpenXRHMDPlugin::PreInit()
 		Extensions.Add(XR_VARJO_QUAD_VIEWS_EXTENSION_NAME);
 	}
 
+	if (HasExtension(XR_KHR_VISIBILITY_MASK_EXTENSION_NAME))
+	{
+		Extensions.Add(XR_KHR_VISIBILITY_MASK_EXTENSION_NAME);
+	}
+
 	// Enable layers, if specified by CVar.
 	// Note: For the validation layer to work on Windows (as of latest OpenXR runtime, August 2019), the following are required:
 	//   1. Download and build the OpenXR SDK from [REDACTED_PROFILE_URL]/OpenXR-SDK-Source (follow instructions at [REDACTED_PROFILE_URL]/OpenXR-SDK-Source/blob/master/BUILDING.md)
@@ -805,7 +810,7 @@ bool FOpenXRHMD::IsActiveThisFrame(class FViewport* InViewport) const
 	return GEngine && GEngine->IsStereoscopic3D(InViewport);
 }
 
-FOpenXRHMD::FOpenXRHMD(const FAutoRegister& AutoRegister, XrInstance InInstance, XrSystemId InSystem, TRefCountPtr<FOpenXRRenderBridge>& InRenderBridge, bool InDepthExtensionSupported)
+FOpenXRHMD::FOpenXRHMD(const FAutoRegister& AutoRegister, XrInstance InInstance, XrSystemId InSystem, TRefCountPtr<FOpenXRRenderBridge>& InRenderBridge, const TSet<FString>& Extensions)
 	: FHeadMountedDisplayBase(nullptr)
 	, FSceneViewExtensionBase(AutoRegister)
 	, bStereoEnabled(false)
@@ -813,8 +818,6 @@ FOpenXRHMD::FOpenXRHMD(const FAutoRegister& AutoRegister, XrInstance InInstance,
 	, bIsReady(false)
 	, bIsRendering(false)
 	, bRunRequested(false)
-	, bDepthExtensionSupported(InDepthExtensionSupported)
-	, bNeedReAllocatedDepth(InDepthExtensionSupported)
 	, CurrentSessionState(XR_SESSION_STATE_UNKNOWN)
 	, Instance(InInstance)
 	, System(InSystem)
@@ -835,6 +838,8 @@ FOpenXRHMD::FOpenXRHMD(const FAutoRegister& AutoRegister, XrInstance InInstance,
 	ViewState.next = nullptr;
 	ViewState.viewStateFlags = 0;
 
+	bNeedReAllocatedDepth = bDepthExtensionSupported = Extensions.Contains(XR_KHR_COMPOSITION_LAYER_DEPTH_EXTENSION_NAME);
+
 	{
 		// Enumerate the viewport configurations
 		uint32 ConfigurationCount;
@@ -877,6 +882,15 @@ FOpenXRHMD::FOpenXRHMD(const FAutoRegister& AutoRegister, XrInstance InInstance,
 
 	// The HMD device does not have an action associated with it
 	ensure(ActionSpaces.Emplace(XR_NULL_HANDLE) == HMDDeviceId);
+
+	if (Extensions.Contains(XR_KHR_VISIBILITY_MASK_EXTENSION_NAME))
+	{
+		FOpenXRHMD* const Self = this;
+		ENQUEUE_RENDER_COMMAND(SetupOpenXROcclusionMeshesCmd)([Self](FRHICommandListImmediate& RHICmdList)
+		{
+			Self->BuildOcclusionMeshes();
+		});
+	}
 }
 
 FOpenXRHMD::~FOpenXRHMD()
@@ -887,6 +901,112 @@ FOpenXRHMD::~FOpenXRHMD()
 	}
 }
 
+void FOpenXRHMD::BuildOcclusionMeshes()
+{
+	HiddenAreaMeshes.SetNum(Views.Num());
+	VisibleAreaMeshes.SetNum(Views.Num());
+
+	bool bSucceeded = true;
+
+	for (int View = 0; View < Views.Num(); ++View)
+	{
+		if (!BuildOcclusionMesh(XR_VISIBILITY_MASK_TYPE_VISIBLE_TRIANGLE_MESH_KHR, View, VisibleAreaMeshes[View]) ||
+			!BuildOcclusionMesh(XR_VISIBILITY_MASK_TYPE_HIDDEN_TRIANGLE_MESH_KHR, View, HiddenAreaMeshes[View]))
+		{
+			bSucceeded = false;
+			break;
+		}
+	}
+
+	if (!bSucceeded)
+	{
+		UE_LOG(LogHMD, Error, TEXT("Failed to create all visibility mask meshes for device/views. Abandoning visibility mask."));
+
+		HiddenAreaMeshes.Empty();
+		VisibleAreaMeshes.Empty();
+	}
+}
+
+bool FOpenXRHMD::BuildOcclusionMesh(XrVisibilityMaskTypeKHR Type, int View, FHMDViewMesh& Mesh)
+{
+	XrVisibilityMaskKHR VisibilityMask = { XR_TYPE_VISIBILITY_MASK_KHR };
+	XR_ENSURE(xrGetVisibilityMaskKHR(Session, SelectedViewConfigurationType, View, Type, &VisibilityMask));
+
+	if (!VisibilityMask.indexCountOutput || (VisibilityMask.indexCountOutput % 3) != 0 || VisibilityMask.vertexCountOutput == 0)
+	{
+		UE_LOG(LogHMD, Error, TEXT("Visibility Mask Mesh returned from runtime is invalid."));
+		return false;
+	}
+
+	FRHIResourceCreateInfo CreateInfo;
+	Mesh.VertexBufferRHI = RHICreateVertexBuffer(sizeof(FFilterVertex) * VisibilityMask.vertexCountOutput, BUF_Static, CreateInfo);
+	void* VertexBufferPtr = RHILockVertexBuffer(Mesh.VertexBufferRHI, 0, sizeof(FFilterVertex) * VisibilityMask.vertexCountOutput, RLM_WriteOnly);
+	FFilterVertex* Vertices = reinterpret_cast<FFilterVertex*>(VertexBufferPtr);
+
+	Mesh.IndexBufferRHI = RHICreateIndexBuffer(sizeof(uint32), sizeof(uint32) * VisibilityMask.indexCountOutput, BUF_Static, CreateInfo);
+	void* IndexBufferPtr = RHILockIndexBuffer(Mesh.IndexBufferRHI, 0, sizeof(uint32) * VisibilityMask.indexCountOutput, RLM_WriteOnly);
+
+	uint32* OutIndices = reinterpret_cast<uint32*>(IndexBufferPtr);
+	TUniquePtr<XrVector2f[]> const OutVertices = MakeUnique<XrVector2f[]>(VisibilityMask.vertexCountOutput);
+
+	VisibilityMask.vertexCapacityInput = VisibilityMask.vertexCountOutput;
+	VisibilityMask.indexCapacityInput = VisibilityMask.indexCountOutput;
+	VisibilityMask.indices = OutIndices;
+	VisibilityMask.vertices = OutVertices.Get();
+
+	xrGetVisibilityMaskKHR(Session, SelectedViewConfigurationType, View, Type, &VisibilityMask);
+
+	ensure(VisibilityMask.vertexCapacityInput == VisibilityMask.vertexCountOutput);
+	ensure(VisibilityMask.indexCapacityInput == VisibilityMask.indexCountOutput);
+
+	if (Type == XR_VISIBILITY_MASK_TYPE_HIDDEN_TRIANGLE_MESH_KHR)
+	{
+		// For the hidden-area mesh, bias z to the near plane so we occlude everything, and re-bias vertices to the -0.5 to 0.5 range in x and y. 
+		for (uint32 VertexIndex = 0; VertexIndex < VisibilityMask.vertexCountOutput; ++VertexIndex)
+		{
+			FFilterVertex& Vertex = Vertices[VertexIndex];
+			const XrVector2f& Position = OutVertices[VertexIndex];
+
+			Vertex.Position.X = (Position.x * 2.0f) - 1.0f;
+			Vertex.Position.Y = (Position.y * 2.0f) - 1.0f;
+			Vertex.Position.Z = 1.0f;
+			Vertex.Position.W = 1.0f;
+			Vertex.UV.X = Vertex.UV.Y = 0.0f;
+		}
+	}
+	else if (Type == XR_VISIBILITY_MASK_TYPE_VISIBLE_TRIANGLE_MESH_KHR)
+	{
+		// For the visible-area mesh, this will be consumed by the post-process pipeline, so set up coordinates in the space they expect 
+		// (x and y range from 0-1, y reversed, z at the far plane).
+		for (uint32 VertexIndex = 0; VertexIndex < VisibilityMask.vertexCountOutput; ++VertexIndex)
+		{
+			FFilterVertex& Vertex = Vertices[VertexIndex];
+			const XrVector2f& Position = OutVertices[VertexIndex];
+
+			Vertex.Position.X = Position.x;
+			Vertex.Position.Y = 1.0f - Position.y;
+			Vertex.Position.Z = 0.0f;
+			Vertex.Position.W = 1.0f;
+
+			Vertex.UV.X = Position.x;
+			Vertex.UV.Y = 1.0f - Position.y;
+		}
+	}
+	else
+	{
+		ensureMsgf(false, TEXT("FOpenXRHMD::BuildOcclusionMesh called with unsupported visibility mask type."));
+	}
+
+	Mesh.NumIndices = VisibilityMask.indexCountOutput;
+	Mesh.NumVertices = VisibilityMask.vertexCountOutput;
+	Mesh.NumTriangles = Mesh.NumIndices / 3;
+
+	RHIUnlockVertexBuffer(Mesh.VertexBufferRHI);
+	RHIUnlockIndexBuffer(Mesh.IndexBufferRHI);
+	
+	return true;
+}
+
 bool FOpenXRHMD::OnStereoStartup()
 {
 	FOpenXRHMD* Self = this;
@@ -1488,19 +1608,42 @@ void FOpenXRHMD::RenderTexture_RenderThread(class FRHICommandListImmediate& RHIC
 	}
 }
 
+bool FOpenXRHMD::HasHiddenAreaMesh() const
+{
+	return HiddenAreaMeshes.Num() > 0;
+}
+
+bool FOpenXRHMD::HasVisibleAreaMesh() const
+{
+	return VisibleAreaMeshes.Num() > 0;
+}
+
 void FOpenXRHMD::DrawHiddenAreaMesh_RenderThread(class FRHICommandList& RHICmdList, EStereoscopicPass StereoPass) const
 {
 	check(IsInRenderingThread());
 	check(StereoPass != eSSP_FULL);
 
-#if 0
 	const uint32 ViewIndex = GetViewIndexForPass(StereoPass);
+	check(ViewIndex < (uint32)HiddenAreaMeshes.Num());
 	const FHMDViewMesh& Mesh = HiddenAreaMeshes[ViewIndex];
 	check(Mesh.IsValid());
 
 	RHICmdList.SetStreamSource(0, Mesh.VertexBufferRHI, 0);
 	RHICmdList.DrawIndexedPrimitive(Mesh.IndexBufferRHI, 0, 0, Mesh.NumVertices, 0, Mesh.NumTriangles, 1);
-#endif
+}
+
+void FOpenXRHMD::DrawVisibleAreaMesh_RenderThread(class FRHICommandList& RHICmdList, EStereoscopicPass StereoPass) const
+{
+	check(IsInRenderingThread());
+	check(StereoPass != eSSP_FULL);
+
+	const uint32 ViewIndex = GetViewIndexForPass(StereoPass);
+	check(ViewIndex < (uint32)VisibleAreaMeshes.Num());
+	const FHMDViewMesh& Mesh = VisibleAreaMeshes[ViewIndex];
+	check(Mesh.IsValid());
+
+	RHICmdList.SetStreamSource(0, Mesh.VertexBufferRHI, 0);
+	RHICmdList.DrawIndexedPrimitive(Mesh.IndexBufferRHI, 0, 0, Mesh.NumVertices, 0, Mesh.NumTriangles, 1);
 }
 
 //---------------------------------------------------
diff --git a/Engine/Plugins/Runtime/OpenXR/Source/OpenXRHMD/Private/OpenXRHMD.h b/Engine/Plugins/Runtime/OpenXR/Source/OpenXRHMD/Private/OpenXRHMD.h
--- a/Engine/Plugins/Runtime/OpenXR/Source/OpenXRHMD/Private/OpenXRHMD.h
+++ b/Engine/Plugins/Runtime/OpenXR/Source/OpenXRHMD/Private/OpenXRHMD.h
@@ -109,6 +109,9 @@ protected:
 	bool ReadNextEvent(XrEventDataBuffer* buffer);
 	void CloseSession();
 
+	void BuildOcclusionMeshes();
+	bool BuildOcclusionMesh(XrVisibilityMaskTypeKHR Type, int View, FHMDViewMesh& Mesh);
+
 public:
 	/** IHeadMountedDisplay interface */
 	virtual bool IsHMDConnected() override { return true; }
@@ -122,8 +125,10 @@ public:
 	virtual bool GetHMDDistortionEnabled(EShadingPath ShadingPath) const override { return false; }
 	virtual FIntRect GetFullFlatEyeRect_RenderThread(FTexture2DRHIRef EyeTexture) const override;
 	virtual void CopyTexture_RenderThread(FRHICommandListImmediate& RHICmdList, FRHITexture2D* SrcTexture, FIntRect SrcRect, FRHITexture2D* DstTexture, FIntRect DstRect, bool bClearBlack, bool bNoAlpha) const override;
-	virtual bool HasHiddenAreaMesh() const override { return false; };
-	virtual void DrawHiddenAreaMesh_RenderThread(class FRHICommandList& RHICmdList, EStereoscopicPass StereoPass) const override;
+	virtual bool HasHiddenAreaMesh() const override final;
+	virtual bool HasVisibleAreaMesh() const override final;
+	virtual void DrawHiddenAreaMesh_RenderThread(class FRHICommandList& RHICmdList, EStereoscopicPass StereoPass) const override final;
+	virtual void DrawVisibleAreaMesh_RenderThread(class FRHICommandList& RHICmdList, EStereoscopicPass StereoPass) const override final;
 	virtual void OnBeginRendering_RenderThread(FRHICommandListImmediate& RHICmdList, FSceneViewFamily& ViewFamily) override;
 	virtual void OnBeginRendering_GameThread() override;
 	virtual void OnLateUpdateApplied_RenderThread(const FTransform& NewRelativeTransform) override;
@@ -164,7 +169,7 @@ public:
 
 public:
 	/** Constructor */
-	FOpenXRHMD(const FAutoRegister&, XrInstance InInstance, XrSystemId InSystem, TRefCountPtr<FOpenXRRenderBridge>& InRenderBridge, bool InDepthExtensionSupported);
+	FOpenXRHMD(const FAutoRegister&, XrInstance InInstance, XrSystemId InSystem, TRefCountPtr<FOpenXRRenderBridge>& InRenderBridge, const TSet<FString>& Extensions);
 
 	/** Destructor */
 	virtual ~FOpenXRHMD();
@@ -217,9 +222,7 @@ private:
 	TArray<XrView>			Views;
 	TArray<XrCompositionLayerProjectionView> ViewsRHI;
 	TArray<XrCompositionLayerDepthInfoKHR> DepthLayersRHI;
-#if 0
-	TArray<FHMDViewMesh>	HiddenAreaMeshes;
-#endif
+
 	TArray<FActionSpace>	ActionSpaces;
 
 	TRefCountPtr<FOpenXRRenderBridge> RenderBridge;
@@ -227,4 +230,7 @@ private:
 
 	FXRSwapChainPtr			Swapchain;
 	FXRSwapChainPtr			DepthSwapchain;
+
+	TArray<FHMDViewMesh>	HiddenAreaMeshes;
+	TArray<FHMDViewMesh>	VisibleAreaMeshes;
 };
