Fix black screen when OpenXR plugin is enabled but VR is disabled

diff --git a/Engine/Plugins/Runtime/OpenXR/Source/OpenXRHMD/Private/OpenXRHMD.h b/Engine/Plugins/Runtime/OpenXR/Source/OpenXRHMD/Private/OpenXRHMD.h
--- a/Engine/Plugins/Runtime/OpenXR/Source/OpenXRHMD/Private/OpenXRHMD.h
+++ b/Engine/Plugins/Runtime/OpenXR/Source/OpenXRHMD/Private/OpenXRHMD.h
@@ -156,7 +156,7 @@ public:
 	virtual bool IsActiveThisFrame(class FViewport* InViewport) const;
 
 	/** IStereoRenderTargetManager */
-	virtual bool ShouldUseSeparateRenderTarget() const override { return true; }
+	virtual bool ShouldUseSeparateRenderTarget() const override { return IsStereoEnabled(); }
 	virtual bool AllocateRenderTargetTexture(uint32 Index, uint32 SizeX, uint32 SizeY, uint8 Format, uint32 NumMips, uint32 Flags, uint32 TargetableTextureFlags, FTexture2DRHIRef& OutTargetableTexture, FTexture2DRHIRef& OutShaderResourceTexture, uint32 NumSamples = 1) override;
 	virtual bool NeedReAllocateDepthTexture(const TRefCountPtr<IPooledRenderTarget>& DepthTarget) override final { return bNeedReAllocatedDepth; }
 	virtual bool AllocateDepthTexture(uint32 Index, uint32 SizeX, uint32 SizeY, uint8 Format, uint32 NumMips, uint32 InTexFlags, uint32 TargetableTextureFlags, FTexture2DRHIRef& OutTargetableTexture, FTexture2DRHIRef& OutShaderResourceTexture, uint32 NumSamples = 1) override final;
