Build the OpenXR occlusion meshes when we create a session
Fixes a crash due to calls being made using an uninitialized session

diff --git a/Engine/Plugins/Runtime/OpenXR/Source/OpenXRHMD/Private/OpenXRHMD.cpp b/Engine/Plugins/Runtime/OpenXR/Source/OpenXRHMD/Private/OpenXRHMD.cpp
--- a/Engine/Plugins/Runtime/OpenXR/Source/OpenXRHMD/Private/OpenXRHMD.cpp
+++ b/Engine/Plugins/Runtime/OpenXR/Source/OpenXRHMD/Private/OpenXRHMD.cpp
@@ -839,6 +839,7 @@ FOpenXRHMD::FOpenXRHMD(const FAutoRegister& AutoRegister, XrInstance InInstance,
 	ViewState.viewStateFlags = 0;
 
 	bNeedReAllocatedDepth = bDepthExtensionSupported = Extensions.Contains(XR_KHR_COMPOSITION_LAYER_DEPTH_EXTENSION_NAME);
+	bHiddenAreaMaskSupported = Extensions.Contains(XR_KHR_VISIBILITY_MASK_EXTENSION_NAME);
 
 	{
 		// Enumerate the viewport configurations
@@ -882,15 +883,6 @@ FOpenXRHMD::FOpenXRHMD(const FAutoRegister& AutoRegister, XrInstance InInstance,
 
 	// The HMD device does not have an action associated with it
 	ensure(ActionSpaces.Emplace(XR_NULL_HANDLE) == HMDDeviceId);
-
-	if (Extensions.Contains(XR_KHR_VISIBILITY_MASK_EXTENSION_NAME))
-	{
-		FOpenXRHMD* const Self = this;
-		ENQUEUE_RENDER_COMMAND(SetupOpenXROcclusionMeshesCmd)([Self](FRHICommandListImmediate& RHICmdList)
-		{
-			Self->BuildOcclusionMeshes();
-		});
-	}
 }
 
 FOpenXRHMD::~FOpenXRHMD()
@@ -1018,6 +1010,11 @@ bool FOpenXRHMD::OnStereoStartup()
 		SessionInfo.createFlags = 0;
 		SessionInfo.systemId = Self->System;
 		XR_ENSURE(xrCreateSession(Self->Instance, &SessionInfo, &Self->Session));
+
+		if (Self->bHiddenAreaMaskSupported)
+		{
+			Self->BuildOcclusionMeshes();
+		}
 	});
 
 	FlushRenderingCommands();
diff --git a/Engine/Plugins/Runtime/OpenXR/Source/OpenXRHMD/Private/OpenXRHMD.h b/Engine/Plugins/Runtime/OpenXR/Source/OpenXRHMD/Private/OpenXRHMD.h
--- a/Engine/Plugins/Runtime/OpenXR/Source/OpenXRHMD/Private/OpenXRHMD.h
+++ b/Engine/Plugins/Runtime/OpenXR/Source/OpenXRHMD/Private/OpenXRHMD.h
@@ -200,6 +200,7 @@ private:
 	bool					bIsRendering;
 	bool					bRunRequested;
 	bool					bDepthExtensionSupported;
+	bool					bHiddenAreaMaskSupported;
 	bool					bNeedReAllocatedDepth;
 	
 	XrSessionState			CurrentSessionState;
