Add depth layer to OpenXR next chain when the extension is supported
Also fixes the depth layer descriptor reversing the depth range incorrectly

diff --git a/Engine/Plugins/Runtime/OpenXR/Source/OpenXRHMD/Private/OpenXRHMD.cpp b/Engine/Plugins/Runtime/OpenXR/Source/OpenXRHMD/Private/OpenXRHMD.cpp
--- a/Engine/Plugins/Runtime/OpenXR/Source/OpenXRHMD/Private/OpenXRHMD.cpp
+++ b/Engine/Plugins/Runtime/OpenXR/Source/OpenXRHMD/Private/OpenXRHMD.cpp
@@ -1280,7 +1280,7 @@ void FOpenXRHMD::OnBeginRendering_RenderThread(FRHICommandListImmediate& RHICmdL
 		XrCompositionLayerDepthInfoKHR& DepthLayer = DepthLayersRHI[ViewIndex];
 
 		Projection.type = XR_TYPE_COMPOSITION_LAYER_PROJECTION_VIEW;
-		Projection.next = nullptr; // &DepthLayer;
+		Projection.next = bDepthExtensionSupported ? &DepthLayer : nullptr;
 		Projection.fov = View.fov;
 		Projection.pose = ToXrPose(ViewTransform * BaseTransform, GetWorldToMetersScale());
 		Projection.subImage.swapchain = static_cast<FOpenXRSwapchain*>(GetSwapchain())->GetHandle();
@@ -1300,10 +1300,10 @@ void FOpenXRHMD::OnBeginRendering_RenderThread(FRHICommandListImmediate& RHICmdL
 			DepthLayer.subImage.swapchain = static_cast<FOpenXRSwapchain*>(GetDepthSwapchain())->GetHandle();
 			DepthLayer.subImage.imageArrayIndex = 0;
 			DepthLayer.subImage.imageRect = Projection.subImage.imageRect;
-			DepthLayer.minDepth = 1.0f;
-			DepthLayer.maxDepth = 0.0f;
-			DepthLayer.nearZ = NearZ;
-			DepthLayer.farZ = FLT_MAX;
+			DepthLayer.minDepth = 0.0f;
+			DepthLayer.maxDepth = 1.0f;
+			DepthLayer.nearZ = FLT_MAX;
+			DepthLayer.farZ = NearZ;
 		}
 
 		OffsetX += Config.recommendedImageRectWidth;
