Always create an sRGB OpenXR swapchain if available
This properly signals to the runtime that we do our own gamma correction

diff --git a/Engine/Plugins/Runtime/OpenXR/Source/OpenXRHMD/Private/OpenXRHMD.cpp b/Engine/Plugins/Runtime/OpenXR/Source/OpenXRHMD/Private/OpenXRHMD.cpp
--- a/Engine/Plugins/Runtime/OpenXR/Source/OpenXRHMD/Private/OpenXRHMD.cpp
+++ b/Engine/Plugins/Runtime/OpenXR/Source/OpenXRHMD/Private/OpenXRHMD.cpp
@@ -1191,6 +1191,9 @@ bool FOpenXRHMD::AllocateRenderTargetTexture(uint32 Index, uint32 SizeX, uint32
 {
 	check(IsInRenderingThread());
 
+	// We need to ensure we can sample from the texture in CopyTexture
+	Flags |= TexCreate_ShaderResource;
+
 	Swapchain = RenderBridge->CreateSwapchain(Session, Format, SizeX, SizeY, NumMips, NumSamples, Flags, TargetableTextureFlags);
 	if (!Swapchain)
 	{
diff --git a/Engine/Plugins/Runtime/OpenXR/Source/OpenXRHMD/Private/OpenXRHMD_Swapchain.cpp b/Engine/Plugins/Runtime/OpenXR/Source/OpenXRHMD/Private/OpenXRHMD_Swapchain.cpp
--- a/Engine/Plugins/Runtime/OpenXR/Source/OpenXRHMD/Private/OpenXRHMD_Swapchain.cpp
+++ b/Engine/Plugins/Runtime/OpenXR/Source/OpenXRHMD/Private/OpenXRHMD_Swapchain.cpp
@@ -108,9 +108,10 @@ uint8 GetNearestSupportedSwapchainFormat(XrSession InSession, uint8 RequestedFor
 	return FallbackFormat;
 }
 
-XrSwapchain CreateSwapchain(XrSession InSession, uint32 PlatformFormat, uint32 SizeX, uint32 SizeY, uint32 NumMips, uint32 NumSamples, uint32 TargetableTextureFlags)
+XrSwapchain CreateSwapchain(XrSession InSession, uint32 PlatformFormat, uint32 SizeX, uint32 SizeY, uint32 NumMips, uint32 NumSamples, uint32 Flags, uint32 TargetableTextureFlags)
 {
-	XrSwapchainUsageFlags Usage = 0;
+	// Need a mutable format so we can reinterpret an sRGB format into a linear format
+	XrSwapchainUsageFlags Usage = XR_SWAPCHAIN_USAGE_MUTABLE_FORMAT_BIT;
 	if (TargetableTextureFlags & TexCreate_RenderTargetable)
 	{
 		Usage |= XR_SWAPCHAIN_USAGE_COLOR_ATTACHMENT_BIT;
@@ -119,11 +120,11 @@ XrSwapchain CreateSwapchain(XrSession InSession, uint32 PlatformFormat, uint32 S
 	{
 		Usage |= XR_SWAPCHAIN_USAGE_DEPTH_STENCIL_ATTACHMENT_BIT;
 	}
-	if (TargetableTextureFlags & TexCreate_ShaderResource)
+	if (Flags & TexCreate_ShaderResource)
 	{
 		Usage |= XR_SWAPCHAIN_USAGE_SAMPLED_BIT;
 	}
-	if (TargetableTextureFlags & TexCreate_UAV)
+	if (Flags & TexCreate_UAV)
 	{
 		Usage |= XR_SWAPCHAIN_USAGE_UNORDERED_ACCESS_BIT;
 	}
@@ -171,7 +172,9 @@ FXRSwapChainPtr CreateSwapchain_D3D11(XrSession InSession, uint8 Format, uint32
 		// We need to convert typeless to typed formats to create a swapchain
 		DXGI_FORMAT PlatformFormat = (DXGI_FORMAT)GPixelFormats[InFormat].PlatformFormat;
 		PlatformFormat = FindDepthStencilDXGIFormat(PlatformFormat);
-		PlatformFormat = FindShaderResourceDXGIFormat(PlatformFormat, Flags & TexCreate_SRGB);
+
+		// UE4 renders a gamma-corrected image so we need to use an sRGB format if available
+		PlatformFormat = FindShaderResourceDXGIFormat(PlatformFormat, true);
 		return PlatformFormat;
 	};
 
@@ -181,7 +184,7 @@ FXRSwapChainPtr CreateSwapchain_D3D11(XrSession InSession, uint8 Format, uint32
 		return nullptr;
 	}
 
-	XrSwapchain Swapchain = CreateSwapchain(InSession, ToPlatformFormat(Format), SizeX, SizeY, NumMips, NumSamples, TargetableTextureFlags);
+	XrSwapchain Swapchain = CreateSwapchain(InSession, ToPlatformFormat(Format), SizeX, SizeY, NumMips, NumSamples, Flags, TargetableTextureFlags);
 	if (!Swapchain)
 	{
 		return nullptr;
@@ -211,7 +214,9 @@ FXRSwapChainPtr CreateSwapchain_D3D12(XrSession InSession, uint8 Format, uint32
 		// We need to convert typeless to typed formats to create a swapchain
 		DXGI_FORMAT PlatformFormat = (DXGI_FORMAT)GPixelFormats[InFormat].PlatformFormat;
 		PlatformFormat = FindDepthStencilDXGIFormat_D3D12(PlatformFormat);
-		PlatformFormat = FindShaderResourceDXGIFormat_D3D12(PlatformFormat, Flags & TexCreate_SRGB);
+
+		// UE4 renders a gamma-corrected image so we need to use an sRGB format if available
+		PlatformFormat = FindShaderResourceDXGIFormat_D3D12(PlatformFormat, true);
 		return PlatformFormat;
 	};
 
@@ -221,7 +226,7 @@ FXRSwapChainPtr CreateSwapchain_D3D12(XrSession InSession, uint8 Format, uint32
 		return nullptr;
 	}
 
-	XrSwapchain Swapchain = CreateSwapchain(InSession, ToPlatformFormat(Format), SizeX, SizeY, NumMips, NumSamples, TargetableTextureFlags);
+	XrSwapchain Swapchain = CreateSwapchain(InSession, ToPlatformFormat(Format), SizeX, SizeY, NumMips, NumSamples, Flags, TargetableTextureFlags);
 	if (!Swapchain)
 	{
 		return nullptr;
@@ -250,7 +255,7 @@ FXRSwapChainPtr CreateSwapchain_OpenGL(XrSession InSession, uint8 Format, uint32
 		return nullptr;
 	}
 
-	XrSwapchain Swapchain = CreateSwapchain(InSession, GPixelFormats[Format].PlatformFormat, SizeX, SizeY, NumMips, NumSamples, TargetableTextureFlags);
+	XrSwapchain Swapchain = CreateSwapchain(InSession, GPixelFormats[Format].PlatformFormat, SizeX, SizeY, NumMips, NumSamples, Flags, TargetableTextureFlags);
 	if (!Swapchain)
 	{
 		return nullptr;
@@ -275,7 +280,8 @@ FXRSwapChainPtr CreateSwapchain_Vulkan(XrSession InSession, uint8 Format, uint32
 {
 	TFunction<uint32(uint8)> ToPlatformFormat = [Flags](uint8 InFormat)
 	{
-		return UEToVkTextureFormat(GPixelFormats[InFormat].UnrealFormat, Flags & TexCreate_SRGB);
+		// UE4 renders a gamma-corrected image so we need to use an sRGB format if available
+		return UEToVkTextureFormat(GPixelFormats[InFormat].UnrealFormat, true);
 	};
 	Format = GetNearestSupportedSwapchainFormat(InSession, Format, ToPlatformFormat);
 	if (!Format)
@@ -283,7 +289,7 @@ FXRSwapChainPtr CreateSwapchain_Vulkan(XrSession InSession, uint8 Format, uint32
 		return nullptr;
 	}
 
-	XrSwapchain Swapchain = CreateSwapchain(InSession, ToPlatformFormat(Format), SizeX, SizeY, NumMips, NumSamples, TargetableTextureFlags);
+	XrSwapchain Swapchain = CreateSwapchain(InSession, ToPlatformFormat(Format), SizeX, SizeY, NumMips, NumSamples, Flags, TargetableTextureFlags);
 	if (!Swapchain)
 	{
 		return nullptr;
