Allocate the OpenXR depth swapchain in AllocateDepthTexture
This ensures the swapchain size matches the DSV

diff --git a/Engine/Plugins/Runtime/OpenXR/Source/OpenXRHMD/Private/OpenXRHMD.cpp b/Engine/Plugins/Runtime/OpenXR/Source/OpenXRHMD/Private/OpenXRHMD.cpp
--- a/Engine/Plugins/Runtime/OpenXR/Source/OpenXRHMD/Private/OpenXRHMD.cpp
+++ b/Engine/Plugins/Runtime/OpenXR/Source/OpenXRHMD/Private/OpenXRHMD.cpp
@@ -678,11 +678,15 @@ void FOpenXRHMD::AdjustViewRect(EStereoscopicPass StereoPass, int32& X, int32& Y
 
 	const XrViewConfigurationView& Config = Configs[ViewIndex];
 
+	FIntPoint ViewRectMin(EForceInit::ForceInitToZero);
 	for (uint32 i = 0; i < ViewIndex; ++i)
 	{
-		X += Configs[i].recommendedImageRectWidth;
+		ViewRectMin.X += Configs[i].recommendedImageRectWidth;
 	}
+	QuantizeSceneBufferSize(ViewRectMin, ViewRectMin);
 
+	X = ViewRectMin.X;
+	Y = ViewRectMin.Y;
 	SizeX = Config.recommendedImageRectWidth;
 	SizeY = Config.recommendedImageRectHeight;
 }
@@ -697,16 +701,16 @@ EStereoscopicPass FOpenXRHMD::GetViewPassForIndex(bool bStereoRequested, uint32
 
 uint32 FOpenXRHMD::GetViewIndexForPass(EStereoscopicPass StereoPassType) const
 {
-	if (IStereoRendering::IsAPrimaryView(StereoPassType) || !IStereoRendering::IsStereoEyeView(StereoPassType))
+	switch (StereoPassType)
 	{
+	case eSSP_LEFT_EYE:
+	case eSSP_FULL:
 		return 0;
-	}
-	else if (IStereoRendering::IsASecondaryView(StereoPassType))
-	{
+
+	case eSSP_RIGHT_EYE:
 		return 1;
-	}
-	else
-	{
+
+	default:
 		return StereoPassType - eSSP_LEFT_EYE;
 	}
 }
@@ -818,6 +822,7 @@ FOpenXRHMD::FOpenXRHMD(const FAutoRegister& AutoRegister, XrInstance InInstance,
 	, bIsReady(false)
 	, bIsRendering(false)
 	, bRunRequested(false)
+	, bNeedReAllocatedDepth(false)
 	, CurrentSessionState(XR_SESSION_STATE_UNKNOWN)
 	, Instance(InInstance)
 	, System(InSystem)
@@ -838,7 +843,7 @@ FOpenXRHMD::FOpenXRHMD(const FAutoRegister& AutoRegister, XrInstance InInstance,
 	ViewState.next = nullptr;
 	ViewState.viewStateFlags = 0;
 
-	bNeedReAllocatedDepth = bDepthExtensionSupported = Extensions.Contains(XR_KHR_COMPOSITION_LAYER_DEPTH_EXTENSION_NAME);
+	bDepthExtensionSupported = Extensions.Contains(XR_KHR_COMPOSITION_LAYER_DEPTH_EXTENSION_NAME);
 	bHiddenAreaMaskSupported = Extensions.Contains(XR_KHR_VISIBILITY_MASK_EXTENSION_NAME);
 
 	{
@@ -1203,27 +1208,27 @@ bool FOpenXRHMD::AllocateRenderTargetTexture(uint32 Index, uint32 SizeX, uint32
 	// Grab the presentation texture out of the swapchain.
 	OutTargetableTexture = OutShaderResourceTexture = (FTexture2DRHIRef&)Swapchain->GetTextureRef();
 
-	if (bDepthExtensionSupported)
-	{
-		// Allocate the depth buffer swapchain while we're here.
-		DepthSwapchain = RenderBridge->CreateSwapchain(Session, PF_DepthStencil, SizeX, SizeY, NumMips, NumSamples, 0, TexCreate_DepthStencilTargetable);
-		if (!DepthSwapchain)
-		{
-			return false;
-		}
-		bNeedReAllocatedDepth = false;
-	}
+	bNeedReAllocatedDepth = bDepthExtensionSupported;
 
 	return true;
 }
 
 bool FOpenXRHMD::AllocateDepthTexture(uint32 Index, uint32 SizeX, uint32 SizeY, uint8 Format, uint32 NumMips, uint32 InTexFlags, uint32 TargetableTextureFlags, FTexture2DRHIRef& OutTargetableTexture, FTexture2DRHIRef& OutShaderResourceTexture, uint32 NumSamples)
 {
-	if (!DepthSwapchain.IsValid())
+	// FIXME: UE4 constantly calls this function even when there is no reason to reallocate the depth texture
+	if (!bDepthExtensionSupported || !bNeedReAllocatedDepth)
 	{
 		return false;
 	}
 
+	DepthSwapchain = RenderBridge->CreateSwapchain(Session, PF_DepthStencil, SizeX, SizeY, FMath::Max(NumMips, 1u), NumSamples, 0, TexCreate_DepthStencilTargetable);
+	if (!DepthSwapchain)
+	{
+		return false;
+	}
+
+	bNeedReAllocatedDepth = false;
+
 	OutTargetableTexture = OutShaderResourceTexture = (FTexture2DRHIRef&)DepthSwapchain->GetTextureRef();
 
 	return true;
@@ -1242,7 +1247,7 @@ void FOpenXRHMD::OnBeginRendering_RenderThread(FRHICommandListImmediate& RHICmdL
 			bIsRendering = true;
 
 			Swapchain->IncrementSwapChainIndex_RHIThread(FrameStateRHI.predictedDisplayPeriod);
-			if (bDepthExtensionSupported)
+			if (bDepthExtensionSupported && !bNeedReAllocatedDepth)
 			{
 				ensure(DepthSwapchain != nullptr);
 				DepthSwapchain->IncrementSwapChainIndex_RHIThread(FrameStateRHI.predictedDisplayPeriod);
@@ -1268,35 +1273,34 @@ void FOpenXRHMD::OnBeginRendering_RenderThread(FRHICommandListImmediate& RHICmdL
 	ViewsRHI.SetNum(Views.Num());
 	DepthLayersRHI.SetNum(Views.Num());
 
-	int32 OffsetX = 0;
-
-	const float WorldScale = GetWorldToMetersScale() * (1.0 / 100.0f); // physical scale is 100 UUs/meter
-	float NearZ = GNearClippingPlane * WorldScale;
+	FIntPoint ViewRectMin(EForceInit::ForceInitToZero);
+	float NearZ = GNearClippingPlane / GetWorldToMetersScale();
 
 	for (int32 ViewIndex = 0; ViewIndex < Views.Num(); ViewIndex++)
 	{
 		const XrView& View = Views[ViewIndex];
 		const XrViewConfigurationView& Config = Configs[ViewIndex];
+
+		FIntRect ViewRect;
+		QuantizeSceneBufferSize(ViewRectMin, ViewRect.Min);
+		ViewRect.Max = ViewRect.Min + FIntPoint(Config.recommendedImageRectWidth, Config.recommendedImageRectHeight);
 		FTransform ViewTransform = ToFTransform(View.pose, GetWorldToMetersScale());
 
 		XrCompositionLayerProjectionView& Projection = ViewsRHI[ViewIndex];
 		XrCompositionLayerDepthInfoKHR& DepthLayer = DepthLayersRHI[ViewIndex];
 
 		Projection.type = XR_TYPE_COMPOSITION_LAYER_PROJECTION_VIEW;
-		Projection.next = bDepthExtensionSupported ? &DepthLayer : nullptr;
+		Projection.next = nullptr;
 		Projection.fov = View.fov;
 		Projection.pose = ToXrPose(ViewTransform * BaseTransform, GetWorldToMetersScale());
 		Projection.subImage.swapchain = static_cast<FOpenXRSwapchain*>(GetSwapchain())->GetHandle();
 		Projection.subImage.imageArrayIndex = 0;
 		Projection.subImage.imageRect = {
-			{ OffsetX, 0 },
-			{
-				(int32)Config.recommendedImageRectWidth,
-				(int32)Config.recommendedImageRectHeight
-			}
+			{ ViewRect.Min.X, ViewRect.Min.Y },
+			{ ViewRect.Width(), ViewRect.Height() }
 		};
 
-		if (bDepthExtensionSupported)
+		if (bDepthExtensionSupported && !bNeedReAllocatedDepth)
 		{
 			DepthLayer.type = XR_TYPE_COMPOSITION_LAYER_DEPTH_INFO_KHR;
 			DepthLayer.next = nullptr;
@@ -1307,9 +1311,11 @@ void FOpenXRHMD::OnBeginRendering_RenderThread(FRHICommandListImmediate& RHICmdL
 			DepthLayer.maxDepth = 1.0f;
 			DepthLayer.nearZ = FLT_MAX;
 			DepthLayer.farZ = NearZ;
+
+			Projection.next = &DepthLayer;
 		}
 
-		OffsetX += Config.recommendedImageRectWidth;
+		ViewRectMin.X += Config.recommendedImageRectWidth;
 	}
 
 	// Give the RHI thread its own copy of the frame state and tracking space
@@ -1503,11 +1509,11 @@ FIntPoint FOpenXRHMD::GetIdealRenderTargetSize() const
 	{
 		Size.X += (int)Config.recommendedImageRectWidth;
 		Size.Y = FMath::Max(Size.Y, (int)Config.recommendedImageRectHeight);
-	}
 
-	// We always prefer the nearest multiple of 4 for our buffer sizes. Make sure we round up here,
-	// so we're consistent with the rest of the engine in creating our buffers.
-	QuantizeSceneBufferSize(Size, Size);
+		// We always prefer the nearest multiple of 4 for our buffer sizes. Make sure we round up here,
+		// so we're consistent with the rest of the engine in creating our buffers.
+		QuantizeSceneBufferSize(Size, Size);
+	}
 
 	return Size;
 }
@@ -1621,7 +1627,7 @@ bool FOpenXRHMD::HasVisibleAreaMesh() const
 void FOpenXRHMD::DrawHiddenAreaMesh_RenderThread(class FRHICommandList& RHICmdList, EStereoscopicPass StereoPass) const
 {
 	check(IsInRenderingThread());
-	check(IStereoRendering::IsStereoEyeView(StereoPass));
+	check(StereoPass != eSSP_FULL);
 
 	const uint32 ViewIndex = GetViewIndexForPass(StereoPass);
 	check(ViewIndex < (uint32)HiddenAreaMeshes.Num());
@@ -1635,7 +1641,7 @@ void FOpenXRHMD::DrawHiddenAreaMesh_RenderThread(class FRHICommandList& RHICmdLi
 void FOpenXRHMD::DrawVisibleAreaMesh_RenderThread(class FRHICommandList& RHICmdList, EStereoscopicPass StereoPass) const
 {
 	check(IsInRenderingThread());
-	check(IStereoRendering::IsStereoEyeView(StereoPass));
+	check(StereoPass != eSSP_FULL);
 
 	const uint32 ViewIndex = GetViewIndexForPass(StereoPass);
 	check(ViewIndex < (uint32)VisibleAreaMeshes.Num());
