Correct the stereo projection matrix field-of-view
The vertical field-of-view was flipped

diff --git a/Engine/Plugins/Runtime/OpenXR/Source/OpenXRHMD/Private/OpenXRHMD.cpp b/Engine/Plugins/Runtime/OpenXR/Source/OpenXRHMD/Private/OpenXRHMD.cpp
--- a/Engine/Plugins/Runtime/OpenXR/Source/OpenXRHMD/Private/OpenXRHMD.cpp
+++ b/Engine/Plugins/Runtime/OpenXR/Source/OpenXRHMD/Private/OpenXRHMD.cpp
@@ -479,18 +479,18 @@ FMatrix FOpenXRHMD::GetStereoProjectionMatrix(const enum EStereoscopicPass Stere
 
 	Fov.angleUp = tan(Fov.angleUp);
 	Fov.angleDown = tan(Fov.angleDown);
-	Fov.angleLeft = tan(-Fov.angleLeft);
-	Fov.angleRight = tan(-Fov.angleRight);
+	Fov.angleLeft = tan(Fov.angleLeft);
+	Fov.angleRight = tan(Fov.angleRight);
 
-	float SumRL = (Fov.angleLeft + Fov.angleRight);
+	float SumRL = (Fov.angleRight + Fov.angleLeft);
 	float SumTB = (Fov.angleUp + Fov.angleDown);
-	float InvRL = (1.0f / (Fov.angleLeft - Fov.angleRight));
+	float InvRL = (1.0f / (Fov.angleRight - Fov.angleLeft));
 	float InvTB = (1.0f / (Fov.angleUp - Fov.angleDown));
 
 	FMatrix Mat = FMatrix(
 		FPlane((2.0f * InvRL), 0.0f, 0.0f, 0.0f),
 		FPlane(0.0f, (2.0f * InvTB), 0.0f, 0.0f),
-		FPlane((SumRL * InvRL), (SumTB * InvTB), 0.0f, 1.0f),
+		FPlane((SumRL * -InvRL), (SumTB * -InvTB), 0.0f, 1.0f),
 		FPlane(0.0f, 0.0f, ZNear, 0.0f)
 	);
 
