Use the parent for scalars instead of providing suggested bindings for both child paths

diff --git a/Engine/Plugins/Runtime/OpenXR/Source/OpenXRInput/Private/OpenXRInput.cpp b/Engine/Plugins/Runtime/OpenXR/Source/OpenXRInput/Private/OpenXRInput.cpp
--- a/Engine/Plugins/Runtime/OpenXR/Source/OpenXRInput/Private/OpenXRInput.cpp
+++ b/Engine/Plugins/Runtime/OpenXR/Source/OpenXRInput/Private/OpenXRInput.cpp
@@ -176,12 +176,10 @@ FOpenXRInputPlugin::FOpenXRInput::FOpenXRInput(FOpenXRHMD* HMD)
 	Bindings.Add(XrActionSuggestedBinding{ Controllers[EControllerHand::Right].Vibration, GetPath(Instance, "/user/hand/right/output/haptic") });
 
 	InteractionMappings.Add(FGamepadKeyNames::MotionController_Left_Shoulder, GetPath(Instance, "/user/hand/left/input/menu/click"));
-	InteractionMappings.Add(FGamepadKeyNames::MotionController_Left_Trigger, GetPath(Instance, "/user/hand/left/input/trigger/value"));
-	InteractionMappings.Add(FGamepadKeyNames::MotionController_Left_TriggerAxis, GetPath(Instance, "/user/hand/left/input/trigger/value"));
-	InteractionMappings.Add(FGamepadKeyNames::MotionController_Left_Grip1, GetPath(Instance, "/user/hand/left/input/grip/value"));
-	InteractionMappings.Add(FGamepadKeyNames::MotionController_Left_Grip1Axis, GetPath(Instance, "/user/hand/left/input/grip/value"));
-	InteractionMappings.Add(FGamepadKeyNames::MotionController_Left_Grip1, GetPath(Instance, "/user/hand/left/input/grip/click"));
-	InteractionMappings.Add(FGamepadKeyNames::MotionController_Left_Grip1Axis, GetPath(Instance, "/user/hand/left/input/grip/click"));
+	InteractionMappings.Add(FGamepadKeyNames::MotionController_Left_Trigger, GetPath(Instance, "/user/hand/left/input/trigger"));
+	InteractionMappings.Add(FGamepadKeyNames::MotionController_Left_TriggerAxis, GetPath(Instance, "/user/hand/left/input/trigger"));
+	InteractionMappings.Add(FGamepadKeyNames::MotionController_Left_Grip1, GetPath(Instance, "/user/hand/left/input/grip"));
+	InteractionMappings.Add(FGamepadKeyNames::MotionController_Left_Grip1Axis, GetPath(Instance, "/user/hand/left/input/grip"));
 	InteractionMappings.Add(FGamepadKeyNames::MotionController_Left_Thumbstick_X, GetPath(Instance, "/user/hand/left/input/thumbstick/x"));
 	InteractionMappings.Add(FGamepadKeyNames::MotionController_Left_Thumbstick_Y, GetPath(Instance, "/user/hand/left/input/thumbstick/y"));
 	InteractionMappings.Add(FGamepadKeyNames::MotionController_Left_Thumbstick, GetPath(Instance, "/user/hand/left/input/thumbstick/click"));
@@ -189,12 +187,10 @@ FOpenXRInputPlugin::FOpenXRInput::FOpenXRInput(FOpenXRHMD* HMD)
 	InteractionMappings.Add(FGamepadKeyNames::MotionController_Left_FaceButton1, GetPath(Instance, "/user/hand/left/input/y/click"));
 
 	InteractionMappings.Add(FGamepadKeyNames::MotionController_Right_Shoulder, GetPath(Instance, "/user/hand/right/input/menu/click"));
-	InteractionMappings.Add(FGamepadKeyNames::MotionController_Right_Trigger, GetPath(Instance, "/user/hand/right/input/trigger/value"));
-	InteractionMappings.Add(FGamepadKeyNames::MotionController_Right_TriggerAxis, GetPath(Instance, "/user/hand/right/input/trigger/value"));
-	InteractionMappings.Add(FGamepadKeyNames::MotionController_Right_Grip1, GetPath(Instance, "/user/hand/right/input/grip/value"));
-	InteractionMappings.Add(FGamepadKeyNames::MotionController_Right_Grip1Axis, GetPath(Instance, "/user/hand/right/input/grip/value"));
-	InteractionMappings.Add(FGamepadKeyNames::MotionController_Right_Grip1, GetPath(Instance, "/user/hand/right/input/grip/click"));
-	InteractionMappings.Add(FGamepadKeyNames::MotionController_Right_Grip1Axis, GetPath(Instance, "/user/hand/right/input/grip/click"));
+	InteractionMappings.Add(FGamepadKeyNames::MotionController_Right_Trigger, GetPath(Instance, "/user/hand/right/input/trigger"));
+	InteractionMappings.Add(FGamepadKeyNames::MotionController_Right_TriggerAxis, GetPath(Instance, "/user/hand/right/input/trigger"));
+	InteractionMappings.Add(FGamepadKeyNames::MotionController_Right_Grip1, GetPath(Instance, "/user/hand/right/input/grip"));
+	InteractionMappings.Add(FGamepadKeyNames::MotionController_Right_Grip1Axis, GetPath(Instance, "/user/hand/right/input/grip"));
 	InteractionMappings.Add(FGamepadKeyNames::MotionController_Right_Thumbstick_X, GetPath(Instance, "/user/hand/right/input/thumbstick/x"));
 	InteractionMappings.Add(FGamepadKeyNames::MotionController_Right_Thumbstick_Y, GetPath(Instance, "/user/hand/right/input/thumbstick/y"));
 	InteractionMappings.Add(FGamepadKeyNames::MotionController_Right_Thumbstick, GetPath(Instance, "/user/hand/right/input/thumbstick/click"));
