Add patch number to version number reported to OpenXR

diff --git a/Engine/Plugins/Runtime/OpenXR/Source/OpenXRHMD/Private/OpenXRHMD.cpp b/Engine/Plugins/Runtime/OpenXR/Source/OpenXRHMD/Private/OpenXRHMD.cpp
--- a/Engine/Plugins/Runtime/OpenXR/Source/OpenXRHMD/Private/OpenXRHMD.cpp
+++ b/Engine/Plugins/Runtime/OpenXR/Source/OpenXRHMD/Private/OpenXRHMD.cpp
@@ -166,7 +166,7 @@ bool FOpenXRHMDPlugin::PreInit()
 	FPlatformString::Convert(Info.applicationInfo.applicationName, XR_MAX_APPLICATION_NAME_SIZE, GetData(AppName), AppName.Len() + 1);
 	Info.applicationInfo.applicationVersion = 0;
 	FCStringAnsi::Strcpy(Info.applicationInfo.engineName, XR_MAX_ENGINE_NAME_SIZE, "Unreal Engine");
-	Info.applicationInfo.engineVersion = (uint32)FEngineVersion::Current().GetMajor() << 16 | FEngineVersion::Current().GetMinor();
+	Info.applicationInfo.engineVersion = (uint32)((FEngineVersion::Current().GetPatch() & 0xFF) << 24 | (FEngineVersion::Current().GetMajor() & 0xFF) << 16 | FEngineVersion::Current().GetMinor());
 	Info.applicationInfo.apiVersion = XR_CURRENT_API_VERSION;
 	Info.enabledApiLayerCount = 0;
 	Info.enabledApiLayerNames = nullptr;
