Refactor OpenXR version information to match VulkanRHI more closely

diff --git a/Engine/Plugins/Runtime/OpenXR/Source/OpenXRHMD/Private/OpenXRHMD.cpp b/Engine/Plugins/Runtime/OpenXR/Source/OpenXRHMD/Private/OpenXRHMD.cpp
--- a/Engine/Plugins/Runtime/OpenXR/Source/OpenXRHMD/Private/OpenXRHMD.cpp
+++ b/Engine/Plugins/Runtime/OpenXR/Source/OpenXRHMD/Private/OpenXRHMD.cpp
@@ -156,17 +156,25 @@ bool FOpenXRHMDPlugin::PreInit()
 		return false;
 	}
 
-	FString AppName = FApp::GetName();
 	const char* extensions[] = { XR_KHR_D3D11_ENABLE_EXTENSION_NAME };
 
+	// Engine registration can be disabled via console var.
+	auto* CVarDisableEngineAndAppRegistration = IConsoleManager::Get().FindTConsoleVariableDataInt(TEXT("r.DisableEngineAndAppRegistration"));
+	bool bDisableEngineRegistration = (CVarDisableEngineAndAppRegistration && CVarDisableEngineAndAppRegistration->GetValueOnAnyThread() != 0);
+
+	// EngineName will be of the form "UnrealEngine4.21", with the minor version ("21" in this example)
+	// updated with every quarterly release
+	FString EngineName = bDisableEngineRegistration ? FString("") : FApp::GetEpicProductIdentifier() + FEngineVersion::Current().ToString(EVersionComponent::Minor);
+	const TCHAR* ProjectName = bDisableEngineRegistration ? TEXT("") : FApp::GetProjectName();
+
 	XrInstanceCreateInfo Info;
 	Info.type = XR_TYPE_INSTANCE_CREATE_INFO;
 	Info.next = nullptr;
 	Info.createFlags = 0;
-	FPlatformString::Convert(Info.applicationInfo.applicationName, XR_MAX_APPLICATION_NAME_SIZE, GetData(AppName), AppName.Len() + 1);
-	Info.applicationInfo.applicationVersion = 0;
-	FCStringAnsi::Strcpy(Info.applicationInfo.engineName, XR_MAX_ENGINE_NAME_SIZE, "Unreal Engine");
-	Info.applicationInfo.engineVersion = (uint32)((FEngineVersion::Current().GetPatch() & 0xFF) << 24 | (FEngineVersion::Current().GetMajor() & 0xFF) << 16 | FEngineVersion::Current().GetMinor());
+	FTCHARToUTF8_Convert::Convert(Info.applicationInfo.applicationName, XR_MAX_APPLICATION_NAME_SIZE, ProjectName, TCString<TCHAR>::Strlen(ProjectName) + 1);
+	Info.applicationInfo.applicationVersion = FEngineVersion::Current().GetChangelist();
+	FTCHARToUTF8_Convert::Convert(Info.applicationInfo.engineName, XR_MAX_ENGINE_NAME_SIZE, *EngineName, EngineName.Len() + 1);
+	Info.applicationInfo.engineVersion = (uint32)(FEngineVersion::Current().GetMinor() << 16 | FEngineVersion::Current().GetPatch());
 	Info.applicationInfo.apiVersion = XR_CURRENT_API_VERSION;
 	Info.enabledApiLayerCount = 0;
 	Info.enabledApiLayerNames = nullptr;
