Remove deprecated usage of SetRenderTarget APIs from OpenXR and WindowsMixedReality plugins.

diff --git a/Engine/Plugins/Runtime/OpenXR/Source/OpenXRHMD/Private/OpenXRHMD.cpp b/Engine/Plugins/Runtime/OpenXR/Source/OpenXRHMD/Private/OpenXRHMD.cpp
--- a/Engine/Plugins/Runtime/OpenXR/Source/OpenXRHMD/Private/OpenXRHMD.cpp
+++ b/Engine/Plugins/Runtime/OpenXR/Source/OpenXRHMD/Private/OpenXRHMD.cpp
@@ -999,58 +999,60 @@ void FOpenXRHMD::CopyTexture_RenderThread(FRHICommandListImmediate& RHICmdList,
 		VSize = SrcRect.Height() / SrcTextureHeight;
 	}
 
-	PRAGMA_DISABLE_DEPRECATION_WARNINGS
-	SetRenderTarget(RHICmdList, DstTexture, FTextureRHIRef());
-	PRAGMA_ENABLE_DEPRECATION_WARNINGS
-
-	if (bClearBlack)
+	FRHITexture * ColorRT = DstTexture->GetTexture2D();
+	FRHIRenderPassInfo RenderPassInfo(ColorRT, ERenderTargetActions::DontLoad_Store);
+	RHICmdList.BeginRenderPass(RenderPassInfo, TEXT("OpenXRHMD_CopyTexture"));
 	{
-		const FIntRect ClearRect(0, 0, DstTexture->GetSizeX(), DstTexture->GetSizeY());
-		RHICmdList.SetViewport(ClearRect.Min.X, ClearRect.Min.Y, 0, ClearRect.Max.X, ClearRect.Max.Y, 1.0f);
-		DrawClearQuad(RHICmdList, FLinearColor::Black);
-	}
+		if (bClearBlack)
+		{
+			const FIntRect ClearRect(0, 0, DstTexture->GetSizeX(), DstTexture->GetSizeY());
+			RHICmdList.SetViewport(ClearRect.Min.X, ClearRect.Min.Y, 0, ClearRect.Max.X, ClearRect.Max.Y, 1.0f);
+			DrawClearQuad(RHICmdList, FLinearColor::Black);
+		}
 
-	RHICmdList.SetViewport(DstRect.Min.X, DstRect.Min.Y, 0, DstRect.Max.X, DstRect.Max.Y, 1.0f);
+		RHICmdList.SetViewport(DstRect.Min.X, DstRect.Min.Y, 0, DstRect.Max.X, DstRect.Max.Y, 1.0f);
 
-	FGraphicsPipelineStateInitializer GraphicsPSOInit;
-	RHICmdList.ApplyCachedRenderTargets(GraphicsPSOInit);
-	GraphicsPSOInit.BlendState = bNoAlpha ? TStaticBlendState<>::GetRHI() : TStaticBlendState<CW_RGBA, BO_Add, BF_SourceAlpha, BF_InverseSourceAlpha, BO_Add, BF_One, BF_InverseSourceAlpha>::GetRHI();
-	GraphicsPSOInit.RasterizerState = TStaticRasterizerState<>::GetRHI();
-	GraphicsPSOInit.DepthStencilState = TStaticDepthStencilState<false, CF_Always>::GetRHI();
-	GraphicsPSOInit.PrimitiveType = PT_TriangleList;
+		FGraphicsPipelineStateInitializer GraphicsPSOInit;
+		RHICmdList.ApplyCachedRenderTargets(GraphicsPSOInit);
+		GraphicsPSOInit.BlendState = bNoAlpha ? TStaticBlendState<>::GetRHI() : TStaticBlendState<CW_RGBA, BO_Add, BF_SourceAlpha, BF_InverseSourceAlpha, BO_Add, BF_One, BF_InverseSourceAlpha>::GetRHI();
+		GraphicsPSOInit.RasterizerState = TStaticRasterizerState<>::GetRHI();
+		GraphicsPSOInit.DepthStencilState = TStaticDepthStencilState<false, CF_Always>::GetRHI();
+		GraphicsPSOInit.PrimitiveType = PT_TriangleList;
 
-	const auto FeatureLevel = GMaxRHIFeatureLevel;
-	auto ShaderMap = GetGlobalShaderMap(FeatureLevel);
+		const auto FeatureLevel = GMaxRHIFeatureLevel;
+		auto ShaderMap = GetGlobalShaderMap(FeatureLevel);
 
-	TShaderMapRef<FScreenVS> VertexShader(ShaderMap);
-	TShaderMapRef<FScreenPS> PixelShader(ShaderMap);
+		TShaderMapRef<FScreenVS> VertexShader(ShaderMap);
+		TShaderMapRef<FScreenPS> PixelShader(ShaderMap);
 
-	GraphicsPSOInit.BoundShaderState.VertexDeclarationRHI = GFilterVertexDeclaration.VertexDeclarationRHI;
-	GraphicsPSOInit.BoundShaderState.VertexShaderRHI = GETSAFERHISHADER_VERTEX(*VertexShader);
-	GraphicsPSOInit.BoundShaderState.PixelShaderRHI = GETSAFERHISHADER_PIXEL(*PixelShader);
+		GraphicsPSOInit.BoundShaderState.VertexDeclarationRHI = GFilterVertexDeclaration.VertexDeclarationRHI;
+		GraphicsPSOInit.BoundShaderState.VertexShaderRHI = GETSAFERHISHADER_VERTEX(*VertexShader);
+		GraphicsPSOInit.BoundShaderState.PixelShaderRHI = GETSAFERHISHADER_PIXEL(*PixelShader);
 
-	SetGraphicsPipelineState(RHICmdList, GraphicsPSOInit);
+		SetGraphicsPipelineState(RHICmdList, GraphicsPSOInit);
 
-	const bool bSameSize = DstRect.Size() == SrcRect.Size();
-	if (bSameSize)
-	{
-		PixelShader->SetParameters(RHICmdList, TStaticSamplerState<SF_Point>::GetRHI(), SrcTexture);
-	}
-	else
-	{
-		PixelShader->SetParameters(RHICmdList, TStaticSamplerState<SF_Bilinear>::GetRHI(), SrcTexture);
-	}
+		const bool bSameSize = DstRect.Size() == SrcRect.Size();
+		if (bSameSize)
+		{
+			PixelShader->SetParameters(RHICmdList, TStaticSamplerState<SF_Point>::GetRHI(), SrcTexture);
+		}
+		else
+		{
+			PixelShader->SetParameters(RHICmdList, TStaticSamplerState<SF_Bilinear>::GetRHI(), SrcTexture);
+		}
 
-	RendererModule->DrawRectangle(
-		RHICmdList,
-		0, 0,
-		ViewportWidth, ViewportHeight,
-		U, V,
-		USize, VSize,
-		TargetSize,
-		FIntPoint(1, 1),
-		*VertexShader,
-		EDRF_Default);
+		RendererModule->DrawRectangle(
+			RHICmdList,
+			0, 0,
+			ViewportWidth, ViewportHeight,
+			U, V,
+			USize, VSize,
+			TargetSize,
+			FIntPoint(1, 1),
+			*VertexShader,
+			EDRF_Default);
+	}
+	RHICmdList.EndRenderPass();
 }
 
 void FOpenXRHMD::RenderTexture_RenderThread(class FRHICommandListImmediate& RHICmdList, class FRHITexture2D* BackBuffer, class FRHITexture2D* SrcTexture, FVector2D WindowSize) const
diff --git a/Engine/Plugins/Runtime/WindowsMixedReality/Source/WindowsMixedRealityHMD/Private/WindowsMixedRealityHMD.cpp b/Engine/Plugins/Runtime/WindowsMixedReality/Source/WindowsMixedRealityHMD/Private/WindowsMixedRealityHMD.cpp
--- a/Engine/Plugins/Runtime/WindowsMixedReality/Source/WindowsMixedRealityHMD/Private/WindowsMixedRealityHMD.cpp
+++ b/Engine/Plugins/Runtime/WindowsMixedReality/Source/WindowsMixedRealityHMD/Private/WindowsMixedRealityHMD.cpp
@@ -886,42 +886,45 @@ namespace WindowsMixedReality
 
 		const uint32 x = (ViewportWidth - width) * 0.5f;
 		const uint32 y = (ViewportHeight - height) * 0.5f;
-		PRAGMA_DISABLE_DEPRECATION_WARNINGS
-		SetRenderTarget(RHICmdList, BackBuffer, FTextureRHIRef());
-		PRAGMA_ENABLE_DEPRECATION_WARNINGS
-		DrawClearQuad(RHICmdList, FLinearColor(0.0f, 0.0f, 0.0f, 1.0f));
-		RHICmdList.SetViewport(x, y, 0, width + x, height + y, 1.0f);
 
-		FGraphicsPipelineStateInitializer GraphicsPSOInit;
-		RHICmdList.ApplyCachedRenderTargets(GraphicsPSOInit);
-
-		const auto FeatureLevel = GMaxRHIFeatureLevel;
-		auto ShaderMap = GetGlobalShaderMap(FeatureLevel);
-		TShaderMapRef<FScreenVS> VertexShader(ShaderMap);
-		TShaderMapRef<FScreenPS> PixelShader(ShaderMap);
-
-		GraphicsPSOInit.BlendState = TStaticBlendState<>::GetRHI();
-		GraphicsPSOInit.RasterizerState = TStaticRasterizerState<>::GetRHI();
-		GraphicsPSOInit.DepthStencilState = TStaticDepthStencilState<false, CF_Always>::GetRHI();
-		GraphicsPSOInit.BoundShaderState.VertexDeclarationRHI = GFilterVertexDeclaration.VertexDeclarationRHI;
-		GraphicsPSOInit.BoundShaderState.VertexShaderRHI = GETSAFERHISHADER_VERTEX(*VertexShader);
-		GraphicsPSOInit.BoundShaderState.PixelShaderRHI = GETSAFERHISHADER_PIXEL(*PixelShader);
-		GraphicsPSOInit.PrimitiveType = PT_TriangleList;
-
-		SetGraphicsPipelineState(RHICmdList, GraphicsPSOInit);
-
-		PixelShader->SetParameters(RHICmdList, TStaticSamplerState<SF_Bilinear>::GetRHI(), SrcTexture);
-
-		RendererModule->DrawRectangle(
-			RHICmdList,
-			0, 0,
-			ViewportWidth, ViewportHeight,
-			0.0f, 0.0f,
-			0.5f, 1.0f,
-			FIntPoint(ViewportWidth, ViewportHeight),
-			FIntPoint(1, 1),
-			*VertexShader,
-			EDRF_Default);
+		FRHIRenderPassInfo RenderPassInfo(static_cast<FRHITexture*>(BackBuffer), ERenderTargetActions::DontLoad_Store);
+		RHICmdList.BeginRenderPass(RenderPassInfo, TEXT("WindowsMixedReality_RenderTexture"));
+		{
+			DrawClearQuad(RHICmdList, FLinearColor(0.0f, 0.0f, 0.0f, 1.0f));
+			RHICmdList.SetViewport(x, y, 0, width + x, height + y, 1.0f);
+
+			FGraphicsPipelineStateInitializer GraphicsPSOInit;
+			RHICmdList.ApplyCachedRenderTargets(GraphicsPSOInit);
+
+			const auto FeatureLevel = GMaxRHIFeatureLevel;
+			auto ShaderMap = GetGlobalShaderMap(FeatureLevel);
+			TShaderMapRef<FScreenVS> VertexShader(ShaderMap);
+			TShaderMapRef<FScreenPS> PixelShader(ShaderMap);
+
+			GraphicsPSOInit.BlendState = TStaticBlendState<>::GetRHI();
+			GraphicsPSOInit.RasterizerState = TStaticRasterizerState<>::GetRHI();
+			GraphicsPSOInit.DepthStencilState = TStaticDepthStencilState<false, CF_Always>::GetRHI();
+			GraphicsPSOInit.BoundShaderState.VertexDeclarationRHI = GFilterVertexDeclaration.VertexDeclarationRHI;
+			GraphicsPSOInit.BoundShaderState.VertexShaderRHI = GETSAFERHISHADER_VERTEX(*VertexShader);
+			GraphicsPSOInit.BoundShaderState.PixelShaderRHI = GETSAFERHISHADER_PIXEL(*PixelShader);
+			GraphicsPSOInit.PrimitiveType = PT_TriangleList;
+
+			SetGraphicsPipelineState(RHICmdList, GraphicsPSOInit);
+
+			PixelShader->SetParameters(RHICmdList, TStaticSamplerState<SF_Bilinear>::GetRHI(), SrcTexture);
+
+			RendererModule->DrawRectangle(
+				RHICmdList,
+				0, 0,
+				ViewportWidth, ViewportHeight,
+				0.0f, 0.0f,
+				0.5f, 1.0f,
+				FIntPoint(ViewportWidth, ViewportHeight),
+				FIntPoint(1, 1),
+				*VertexShader,
+				EDRF_Default);
+		}
+		RHICmdList.EndRenderPass();
 	}
 
 	// Create a BGRA backbuffer for rendering.
