Move OpenXR plugin swapchain usage over to new XRSwapChain class.

diff --git a/Engine/Plugins/Runtime/Oculus/OculusVR/Source/OculusHMD/Private/OculusHMD_CustomPresent.cpp b/Engine/Plugins/Runtime/Oculus/OculusVR/Source/OculusHMD/Private/OculusHMD_CustomPresent.cpp
--- a/Engine/Plugins/Runtime/Oculus/OculusVR/Source/OculusHMD/Private/OculusHMD_CustomPresent.cpp
+++ b/Engine/Plugins/Runtime/Oculus/OculusVR/Source/OculusHMD/Private/OculusHMD_CustomPresent.cpp
@@ -280,6 +280,7 @@ FXRSwapChainPtr FCustomPresent::CreateSwapChain_RenderThread(uint32 InSizeX, uin
 {
 	CheckInRenderThread();
 
+	// @todo: Once things settle down, the chain target will be created below in the CreateXRSwapChain call, via an RHI "CreateAliasedTexture" call.
 	FTextureRHIRef RHITexture;
 	{
 		RHITexture = CreateTexture_RenderThread(InSizeX, InSizeY, InFormat, InBinding, InNumMips, InNumSamples, InNumSamplesTileMem, InResourceType, InTextures[0], InTexCreateFlags);
@@ -293,7 +294,7 @@ FXRSwapChainPtr FCustomPresent::CreateSwapChain_RenderThread(uint32 InSizeX, uin
 		}
 	}
 
-	return MakeShareable(new FXRSwapChain(RHITexture, MoveTemp(RHITextureSwapChain)));
+	return CreateXRSwapChain(MoveTemp(RHITextureSwapChain), RHITexture);
 }
 
 
diff --git a/Engine/Plugins/Runtime/OpenXR/Source/OpenXRHMD/Private/OpenXRHMD.cpp b/Engine/Plugins/Runtime/OpenXR/Source/OpenXRHMD/Private/OpenXRHMD.cpp
--- a/Engine/Plugins/Runtime/OpenXR/Source/OpenXRHMD/Private/OpenXRHMD.cpp
+++ b/Engine/Plugins/Runtime/OpenXR/Source/OpenXRHMD/Private/OpenXRHMD.cpp
@@ -212,28 +212,16 @@ bool FOpenXRHMDPlugin::PreInit()
 	return true;
 }
 
-FOpenXRHMD::FOpenXRSwapchain::FOpenXRSwapchain(XrSwapchain InSwapchain, FTexture2DRHIParamRef InRHITexture, const TArray<FTexture2DRHIRef>& InRHITextureSwapChain) :
-	Handle(InSwapchain), RHITexture(InRHITexture), RHITextureSwapChain(InRHITextureSwapChain), SwapChainIndex_RenderThread(0), IsAcquired(false)
+FOpenXRHMD::FOpenXRSwapchain::FOpenXRSwapchain(TArray<FTextureRHIRef>&& InRHITextureSwapChain, const FTextureRHIRef & InRHITexture) :
+	FXRSwapChain(MoveTemp(InRHITextureSwapChain), InRHITexture),
+	Handle(nullptr), 
+	IsAcquired(false)
+	
 {
-	IncrementSwapChainIndex_RenderThread(XR_NO_DURATION);
+	IncrementSwapChainIndex_RHIThread((int64)XR_NO_DURATION);
 }
 
-FOpenXRHMD::FOpenXRSwapchain::~FOpenXRSwapchain()
-{
-	if (IsInGameThread())
-	{
-		ExecuteOnRenderThread([this]()
-		{
-			ReleaseResources_RenderThread();
-		});
-	}
-	else
-	{
-		ReleaseResources_RenderThread();
-	}
-}
-
-void FOpenXRHMD::FOpenXRSwapchain::IncrementSwapChainIndex_RenderThread(XrDuration Timeout)
+void FOpenXRHMD::FOpenXRSwapchain::IncrementSwapChainIndex_RHIThread(int64 Timeout)
 {
 	check(IsInRenderingThread());
 
@@ -243,7 +231,7 @@ void FOpenXRHMD::FOpenXRSwapchain::IncrementSwapChainIndex_RenderThread(XrDurati
 	XrSwapchainImageAcquireInfo Info;
 	Info.type = XR_TYPE_SWAPCHAIN_IMAGE_ACQUIRE_INFO;
 	Info.next = nullptr;
-	XR_ENSURE(xrAcquireSwapchainImage(Handle, &Info, &SwapChainIndex_RenderThread));
+	XR_ENSURE(xrAcquireSwapchainImage(Handle, &Info, &SwapChainIndex_RHIThread));
 
 	IsAcquired = true;
 
@@ -253,10 +241,10 @@ void FOpenXRHMD::FOpenXRSwapchain::IncrementSwapChainIndex_RenderThread(XrDurati
 	WaitInfo.timeout = Timeout;
 	XR_ENSURE(xrWaitSwapchainImage(Handle, &WaitInfo));
 
-	GDynamicRHI->RHIAliasTextureResources(RHITexture, RHITextureSwapChain[SwapChainIndex_RenderThread]);
+	GDynamicRHI->RHIAliasTextureResources(RHITexture, RHITextureSwapChain[SwapChainIndex_RHIThread]);
 }
 
-void FOpenXRHMD::FOpenXRSwapchain::ReleaseSwapChainImage_RenderThread()
+void FOpenXRHMD::FOpenXRSwapchain::ReleaseCurrentImage_RHIThread()
 {
 	check(IsInRenderingThread());
 
@@ -271,12 +259,9 @@ void FOpenXRHMD::FOpenXRSwapchain::ReleaseSwapChainImage_RenderThread()
 	IsAcquired = false;
 }
 
-void FOpenXRHMD::FOpenXRSwapchain::ReleaseResources_RenderThread()
+void FOpenXRHMD::FOpenXRSwapchain::ReleaseResources_RHIThread()
 {
-	check(IsInRenderingThread());
-
-	RHITexture = nullptr;
-	RHITextureSwapChain.Empty();
+	FXRSwapChain::ReleaseResources_RHIThread();
 	xrDestroySwapchain(Handle);
 }
 
@@ -785,15 +770,19 @@ bool FOpenXRHMD::AllocateRenderTargetTexture(uint32 Index, uint32 SizeX, uint32
 	XR_ENSURE(xrEnumerateSwapchainImages(SwapchainHandle, ChainCount, &ChainCount, reinterpret_cast<XrSwapchainImageBaseHeader*>(Images.GetData())));
 
 	FD3D11DynamicRHI* DynamicRHI = static_cast<FD3D11DynamicRHI*>(GDynamicRHI);
-	TArray<FTexture2DRHIRef> TextureChain;
-	FTexture2DRHIRef ChainTarget = DynamicRHI->RHICreateTexture2DFromResource(SwapchainFormat->PixelFormat, TexCreate_RenderTargetable | TexCreate_ShaderResource, FClearValueBinding::Black, Images[0].texture);
+	TArray<FTextureRHIRef> TextureChain;
+	// @todo: Once things settle down, the chain target will be created below in the CreateXRSwapChain call, via an RHI "CreateAliasedTexture" call.
+	FTextureRHIRef ChainTarget = static_cast<FTextureRHIRef>(DynamicRHI->RHICreateTexture2DFromResource(SwapchainFormat->PixelFormat, TexCreate_RenderTargetable | TexCreate_ShaderResource, FClearValueBinding::Black, Images[0].texture));
 	for (const auto& Image : Images)
 	{
-		TextureChain.Add(DynamicRHI->RHICreateTexture2DFromResource(SwapchainFormat->PixelFormat, TexCreate_RenderTargetable | TexCreate_ShaderResource, FClearValueBinding::Black, Image.texture));
+		TextureChain.Add(static_cast<FTextureRHIRef>(DynamicRHI->RHICreateTexture2DFromResource(SwapchainFormat->PixelFormat, TexCreate_RenderTargetable | TexCreate_ShaderResource, FClearValueBinding::Black, Image.texture)));
 	}
 
-	Swapchain = MakeShareable(new FOpenXRSwapchain(SwapchainHandle, ChainTarget, TextureChain));
-	OutTargetableTexture = OutShaderResourceTexture = ChainTarget;
+	Swapchain = CreateXRSwapChain<FOpenXRSwapchain>(MoveTemp(TextureChain), ChainTarget);
+	static_cast<FOpenXRSwapchain*>(GetSwapchain())->SetHandle(SwapchainHandle);
+	
+	// Grab the presentation texture out of the swapchain.
+	OutTargetableTexture = OutShaderResourceTexture = (FTexture2DRHIRef&)Swapchain->GetTextureRef();
 	return true;
 }
 
@@ -808,7 +797,7 @@ void FOpenXRHMD::OnBeginRendering_RenderThread(FRHICommandListImmediate& RHICmdL
 	check(MainView);
 	BaseTransform = FTransform(MainView->BaseHmdOrientation, MainView->BaseHmdLocation);
 
-	Swapchain->IncrementSwapChainIndex_RenderThread(FrameStateRHI.predictedDisplayPeriod);
+	Swapchain->IncrementSwapChainIndex_RHIThread(FrameStateRHI.predictedDisplayPeriod);
 
 	ViewsRHI.SetNum(Views.Num());
 	int32 OffsetX = 0;
@@ -823,7 +812,7 @@ void FOpenXRHMD::OnBeginRendering_RenderThread(FRHICommandListImmediate& RHICmdL
 		Projection.next = nullptr;
 		Projection.fov = View.fov;
 		Projection.pose = ToXrPose(ViewTransform * BaseTransform, GetWorldToMetersScale());
-		Projection.subImage.swapchain = Swapchain->Handle;
+		Projection.subImage.swapchain = static_cast<FOpenXRSwapchain*>(GetSwapchain())->GetHandle();
 		Projection.subImage.imageArrayIndex = 0;
 		Projection.subImage.imageRect = {
 			{ OffsetX, 0 },
@@ -934,7 +923,7 @@ void FOpenXRHMD::FinishRendering()
 	Layer.viewCount = ViewsRHI.Num();
 	Layer.views = ViewsRHI.GetData();
 
-	Swapchain->ReleaseSwapChainImage_RenderThread();
+	Swapchain->ReleaseCurrentImage_RHIThread();
 
 	XrFrameEndInfo EndInfo;
 	XrCompositionLayerBaseHeader* Headers[1] = { reinterpret_cast<XrCompositionLayerBaseHeader*>(&Layer) };
diff --git a/Engine/Plugins/Runtime/OpenXR/Source/OpenXRHMD/Private/OpenXRHMD.h b/Engine/Plugins/Runtime/OpenXR/Source/OpenXRHMD/Private/OpenXRHMD.h
--- a/Engine/Plugins/Runtime/OpenXR/Source/OpenXRHMD/Private/OpenXRHMD.h
+++ b/Engine/Plugins/Runtime/OpenXR/Source/OpenXRHMD/Private/OpenXRHMD.h
@@ -7,6 +7,7 @@
 #include "XRTrackingSystemBase.h"
 #include "XRRenderTargetManager.h"
 #include "XRRenderBridge.h"
+#include "XRSwapChain.h"
 #include "SceneViewExtension.h"
 #include "DefaultSpectatorScreenController.h"
 
@@ -23,28 +24,24 @@ class UCanvas;
 class FOpenXRHMD : public FHeadMountedDisplayBase, public FXRRenderTargetManager, public FSceneViewExtensionBase
 {
 public:
-	class FOpenXRSwapchain : public TSharedFromThis<FOpenXRSwapchain, ESPMode::ThreadSafe>
+	class FOpenXRSwapchain : public FXRSwapChain
 	{
 	public:
-		FOpenXRSwapchain(XrSwapchain InSwapchain, FTexture2DRHIParamRef InRHITexture, const TArray<FTexture2DRHIRef>& InRHITextureSwapChain);
-		virtual ~FOpenXRSwapchain();
+		FOpenXRSwapchain(TArray<FTextureRHIRef>&& InRHITextureSwapChain, const FTextureRHIRef & InRHITexture);
+		virtual ~FOpenXRSwapchain() {}
 
-		FRHITexture* GetTexture() const { return RHITexture.GetReference(); }
-		FRHITexture2D* GetTexture2D() const { return RHITexture->GetTexture2D(); }
-		FRHITextureCube* GetTextureCube() const { return RHITexture->GetTextureCube(); }
-		uint32 GetSwapchainLength() const { return (uint32)RHITextureSwapChain.Num(); }
+		virtual void IncrementSwapChainIndex_RHIThread(int64 Timeout) override final;
+		virtual void ReleaseCurrentImage_RHIThread() override final;
 
-		uint32 GetSwapchainIndex_RenderThread() { return SwapChainIndex_RenderThread; }
-		void IncrementSwapChainIndex_RenderThread(XrDuration Timeout);
-		void ReleaseSwapChainImage_RenderThread();
+		void SetHandle(XrSwapchain InHandle) { Handle = InHandle; }
+		XrSwapchain GetHandle() { return Handle; }
 		
-		XrSwapchain Handle;
 	protected:
-		void ReleaseResources_RenderThread();
+		virtual void ReleaseResources_RHIThread() override final;
+
+		FXRSwapChainPtr mSwapChain;
+		XrSwapchain Handle;
 
-		FTextureRHIRef RHITexture;
-		TArray<FTexture2DRHIRef> RHITextureSwapChain;
-		uint32 SwapChainIndex_RenderThread;
 		bool IsAcquired;
 	};
 
@@ -164,7 +161,7 @@ public:
 
 	OPENXRHMD_API int32 AddActionDevice(XrAction Action);
 
-	FOpenXRSwapchain* GetSwapchain() { return Swapchain.Get(); }
+	FXRSwapChain* GetSwapchain() { return Swapchain.Get(); }
 	XrInstance GetInstance() { return Instance; }
 	XrSystemId GetSystem() { return System; }
 	XrSession GetSession() { return Session; }
@@ -200,6 +197,5 @@ private:
 	TRefCountPtr<FXRRenderBridge> RenderBridge;
 	IRendererModule*		RendererModule;
 
-	// TODO: Needs to be part of a custom present class
-	TSharedPtr<FOpenXRSwapchain, ESPMode::ThreadSafe> Swapchain;
+	FXRSwapChainPtr			Swapchain;
 };
diff --git a/Engine/Source/Runtime/HeadMountedDisplay/Private/XRSwapChain.cpp b/Engine/Source/Runtime/HeadMountedDisplay/Private/XRSwapChain.cpp
--- a/Engine/Source/Runtime/HeadMountedDisplay/Private/XRSwapChain.cpp
+++ b/Engine/Source/Runtime/HeadMountedDisplay/Private/XRSwapChain.cpp
@@ -60,11 +60,15 @@ FORCEINLINE void CheckInRHIThread()
 #endif
 }
 
-FXRSwapChain::FXRSwapChain(FTextureRHIParamRef InRHITexture, TArray<FTextureRHIRef>&& InRHITextureSwapChain)
-	: RHITexture(InRHITexture)
+FXRSwapChain::FXRSwapChain(TArray<FTextureRHIRef>&& InRHITextureSwapChain, const FTextureRHIRef & AliasedTexture)
+	: RHITexture(AliasedTexture)
 	, RHITextureSwapChain(InRHITextureSwapChain)
 	, SwapChainIndex_RHIThread(0)
 {
+	check(RHITexture);
+	// @todo: The *correct* way to create the RHITexture object (that's just an alias over swap-chain textures)
+	//		  would be via a new RHI API to create an aliased texture. For now we just let the clients of this class
+	//		  create themselves (since they all create differently).
 	RHITexture->SetName(TEXT("XRSwapChainAliasedTexture"));
 	for (int ChainElement = 0; ChainElement < RHITextureSwapChain.Num(); ++ChainElement)
 	{
@@ -108,7 +112,7 @@ void FXRSwapChain::GenerateMips_RenderThread(FRHICommandListImmediate& RHICmdLis
 }
 
 
-void FXRSwapChain::IncrementSwapChainIndex_RHIThread()
+void FXRSwapChain::IncrementSwapChainIndex_RHIThread(int64 /* TimeoutNanoseconds */)
 {
 	CheckInRHIThread();
 
diff --git a/Engine/Source/Runtime/HeadMountedDisplay/Public/XRSwapChain.h b/Engine/Source/Runtime/HeadMountedDisplay/Public/XRSwapChain.h
--- a/Engine/Source/Runtime/HeadMountedDisplay/Public/XRSwapChain.h
+++ b/Engine/Source/Runtime/HeadMountedDisplay/Public/XRSwapChain.h
@@ -11,9 +11,10 @@
 class HEADMOUNTEDDISPLAY_API FXRSwapChain : public TSharedFromThis<FXRSwapChain, ESPMode::ThreadSafe>
 {
 public:
-	FXRSwapChain(FTextureRHIParamRef InRHITexture, TArray<FTextureRHIRef>&& InRHITextureSwapChain);
+	FXRSwapChain(TArray<FTextureRHIRef>&& InRHITextureSwapChain, const FTextureRHIRef & AliasedTexture);
 	virtual ~FXRSwapChain();
 
+	const FTextureRHIRef& GetTextureRef() const { return RHITexture; }
 	FRHITexture* GetTexture() const { return RHITexture.GetReference(); }
 	FRHITexture2D* GetTexture2D() const { return RHITexture->GetTexture2D(); }
 	FRHITexture2DArray* GetTexture2DArray() const { return RHITexture->GetTexture2DArray(); }
@@ -23,7 +24,7 @@ public:
 	void GenerateMips_RenderThread(FRHICommandListImmediate& RHICmdList);
 	uint32 GetSwapChainIndex_RHIThread() { return SwapChainIndex_RHIThread; }
 
-	virtual void IncrementSwapChainIndex_RHIThread();
+	virtual void IncrementSwapChainIndex_RHIThread(int64 TimeoutNanoseconds = 0); // Default to no timeout (immediate).
 	virtual void ReleaseCurrentImage_RHIThread() {}		
 
 protected:
@@ -36,3 +37,12 @@ protected:
 
 typedef TSharedPtr<FXRSwapChain, ESPMode::ThreadSafe> FXRSwapChainPtr;
 
+template<typename T = FXRSwapChain>
+FXRSwapChainPtr CreateXRSwapChain(TArray<FTextureRHIRef>&& InRHITextureSwapChain, const FTextureRHIRef & AliasedTexture)
+{
+	check(InRHITextureSwapChain.Num() >= 1);
+	return MakeShareable(new T(MoveTemp(InRHITextureSwapChain), AliasedTexture));
+}
+
+
+
