Fix crash in initialization of OpenXR swapchain after refactor to leverage XRSwapChain base class.

diff --git a/Engine/Plugins/Runtime/OpenXR/Source/OpenXRHMD/Private/OpenXRHMD.cpp b/Engine/Plugins/Runtime/OpenXR/Source/OpenXRHMD/Private/OpenXRHMD.cpp
--- a/Engine/Plugins/Runtime/OpenXR/Source/OpenXRHMD/Private/OpenXRHMD.cpp
+++ b/Engine/Plugins/Runtime/OpenXR/Source/OpenXRHMD/Private/OpenXRHMD.cpp
@@ -212,9 +212,9 @@ bool FOpenXRHMDPlugin::PreInit()
 	return true;
 }
 
-FOpenXRHMD::FOpenXRSwapchain::FOpenXRSwapchain(TArray<FTextureRHIRef>&& InRHITextureSwapChain, const FTextureRHIRef & InRHITexture) :
+FOpenXRHMD::FOpenXRSwapchain::FOpenXRSwapchain(TArray<FTextureRHIRef>&& InRHITextureSwapChain, const FTextureRHIRef & InRHITexture, XrSwapchain InHandle) :
 	FXRSwapChain(MoveTemp(InRHITextureSwapChain), InRHITexture),
-	Handle(nullptr), 
+	Handle(InHandle), 
 	IsAcquired(false)
 	
 {
@@ -778,8 +778,7 @@ bool FOpenXRHMD::AllocateRenderTargetTexture(uint32 Index, uint32 SizeX, uint32
 		TextureChain.Add(static_cast<FTextureRHIRef>(DynamicRHI->RHICreateTexture2DFromResource(SwapchainFormat->PixelFormat, TexCreate_RenderTargetable | TexCreate_ShaderResource, FClearValueBinding::Black, Image.texture)));
 	}
 
-	Swapchain = CreateXRSwapChain<FOpenXRSwapchain>(MoveTemp(TextureChain), ChainTarget);
-	static_cast<FOpenXRSwapchain*>(GetSwapchain())->SetHandle(SwapchainHandle);
+	Swapchain = CreateXRSwapChain<FOpenXRSwapchain>(MoveTemp(TextureChain), ChainTarget, SwapchainHandle);
 	
 	// Grab the presentation texture out of the swapchain.
 	OutTargetableTexture = OutShaderResourceTexture = (FTexture2DRHIRef&)Swapchain->GetTextureRef();
diff --git a/Engine/Plugins/Runtime/OpenXR/Source/OpenXRHMD/Private/OpenXRHMD.h b/Engine/Plugins/Runtime/OpenXR/Source/OpenXRHMD/Private/OpenXRHMD.h
--- a/Engine/Plugins/Runtime/OpenXR/Source/OpenXRHMD/Private/OpenXRHMD.h
+++ b/Engine/Plugins/Runtime/OpenXR/Source/OpenXRHMD/Private/OpenXRHMD.h
@@ -27,13 +27,12 @@ public:
 	class FOpenXRSwapchain : public FXRSwapChain
 	{
 	public:
-		FOpenXRSwapchain(TArray<FTextureRHIRef>&& InRHITextureSwapChain, const FTextureRHIRef & InRHITexture);
+		FOpenXRSwapchain(TArray<FTextureRHIRef>&& InRHITextureSwapChain, const FTextureRHIRef & InRHITexture, XrSwapchain InHandle);
 		virtual ~FOpenXRSwapchain() {}
 
 		virtual void IncrementSwapChainIndex_RHIThread(int64 Timeout) override final;
 		virtual void ReleaseCurrentImage_RHIThread() override final;
 
-		void SetHandle(XrSwapchain InHandle) { Handle = InHandle; }
 		XrSwapchain GetHandle() { return Handle; }
 		
 	protected:
diff --git a/Engine/Source/Runtime/HeadMountedDisplay/Public/XRSwapChain.h b/Engine/Source/Runtime/HeadMountedDisplay/Public/XRSwapChain.h
--- a/Engine/Source/Runtime/HeadMountedDisplay/Public/XRSwapChain.h
+++ b/Engine/Source/Runtime/HeadMountedDisplay/Public/XRSwapChain.h
@@ -37,11 +37,11 @@ protected:
 
 typedef TSharedPtr<FXRSwapChain, ESPMode::ThreadSafe> FXRSwapChainPtr;
 
-template<typename T = FXRSwapChain>
-FXRSwapChainPtr CreateXRSwapChain(TArray<FTextureRHIRef>&& InRHITextureSwapChain, const FTextureRHIRef & AliasedTexture)
+template<typename T = FXRSwapChain, typename ... Types>
+FXRSwapChainPtr CreateXRSwapChain(TArray<FTextureRHIRef>&& InRHITextureSwapChain, const FTextureRHIRef & InRHIAliasedTexture, Types... InExtraParameters)
 {
 	check(InRHITextureSwapChain.Num() >= 1);
-	return MakeShareable(new T(MoveTemp(InRHITextureSwapChain), AliasedTexture));
+	return MakeShareable(new T(MoveTemp(InRHITextureSwapChain), InRHIAliasedTexture, InExtraParameters...));
 }
 
 
