Ensure FrameState data is initialized (type and next ptr) before xrWaitFrame - can fail on certain runtimes if not.

diff --git a/Engine/Plugins/Runtime/OpenXR/Source/OpenXRHMD/Private/OpenXRHMD.cpp b/Engine/Plugins/Runtime/OpenXR/Source/OpenXRHMD/Private/OpenXRHMD.cpp
--- a/Engine/Plugins/Runtime/OpenXR/Source/OpenXRHMD/Private/OpenXRHMD.cpp
+++ b/Engine/Plugins/Runtime/OpenXR/Source/OpenXRHMD/Private/OpenXRHMD.cpp
@@ -557,8 +557,17 @@ FOpenXRHMD::FOpenXRHMD(const FAutoRegister& AutoRegister, XrInstance InInstance,
 	, StageSpace(XR_NULL_HANDLE)
 	, TrackingSpaceType(XR_REFERENCE_SPACE_TYPE_STAGE)
 	, RenderBridge(nullptr)
+	, RendererModule(nullptr)
 	, Swapchain(XR_NULL_HANDLE)
 {
+	FrameState.type = XR_TYPE_FRAME_STATE;
+	FrameState.next = nullptr;
+	FrameState.predictedDisplayPeriod = FrameState.predictedDisplayTime = 0;
+	FrameStateRHI = FrameState;
+
+	ViewState.type = XR_TYPE_VIEW_STATE;
+	ViewState.next = nullptr;
+	ViewState.viewStateFlags = 0;
 	{
 		// Enumerate the viewport configurations
 		uint32 ConfigurationCount;
@@ -845,9 +854,15 @@ void FOpenXRHMD::OnBeginRendering_GameThread()
 	XrFrameWaitInfo WaitInfo;
 	WaitInfo.type = XR_TYPE_FRAME_WAIT_INFO;
 	WaitInfo.next = nullptr;
+
+	// Pass in a clean FrameState (some implementations may fail if this is just uninitialized).
+	FrameState.type = XR_TYPE_FRAME_STATE;
+	FrameState.next = nullptr;
+	FrameState.predictedDisplayPeriod = FrameStateRHI.predictedDisplayTime = 0;
+
 	XR_ENSURE(xrWaitFrame(Session, &WaitInfo, &FrameState));
 
-	uint32_t ViewCount;
+	uint32_t ViewCount = 0;
 	XrViewLocateInfo ViewInfo;
 	ViewInfo.type = XR_TYPE_VIEW_LOCATE_INFO;
 	ViewInfo.next = nullptr;
