Add depth support to OpenXR plugin - includes plumbing depth through to present, and managing depth swapchain.

diff --git a/Engine/Plugins/Runtime/OpenXR/Source/OpenXRHMD/Private/OpenXRHMD.cpp b/Engine/Plugins/Runtime/OpenXR/Source/OpenXRHMD/Private/OpenXRHMD.cpp
--- a/Engine/Plugins/Runtime/OpenXR/Source/OpenXRHMD/Private/OpenXRHMD.cpp
+++ b/Engine/Plugins/Runtime/OpenXR/Source/OpenXRHMD/Private/OpenXRHMD.cpp
@@ -35,7 +35,12 @@ namespace {
 		{ DXGI_FORMAT_B8G8R8A8_UNORM, PF_B8G8R8A8 },
 		{ DXGI_FORMAT_R8G8B8A8_UNORM_SRGB, PF_R8G8B8A8 },
 		{ DXGI_FORMAT_B8G8R8A8_UNORM_SRGB, PF_B8G8R8A8 },
-    };
+	};
+
+	const FormatMap SupportedDepthSwapchainFormats[] = {
+		{ DXGI_FORMAT_D32_FLOAT_S8X24_UINT, PF_DepthStencil },
+		{ DXGI_FORMAT_D24_UNORM_S8_UINT, PF_D24 },
+	};
 
 	/** Helper function for acquiring the appropriate FSceneViewport */
 	FSceneViewport* FindSceneViewport()
@@ -558,7 +563,6 @@ FOpenXRHMD::FOpenXRHMD(const FAutoRegister& AutoRegister, XrInstance InInstance,
 	, TrackingSpaceType(XR_REFERENCE_SPACE_TYPE_STAGE)
 	, RenderBridge(nullptr)
 	, RendererModule(nullptr)
-	, Swapchain(XR_NULL_HANDLE)
 {
 	FrameState.type = XR_TYPE_FRAME_STATE;
 	FrameState.next = nullptr;
@@ -791,6 +795,92 @@ bool FOpenXRHMD::AllocateRenderTargetTexture(uint32 Index, uint32 SizeX, uint32
 	
 	// Grab the presentation texture out of the swapchain.
 	OutTargetableTexture = OutShaderResourceTexture = (FTexture2DRHIRef&)Swapchain->GetTextureRef();
+
+	// Allocate the depth buffer swapchain while we're here.
+	return AllocateDepthSwapChain(Index, SizeX, SizeY, NumMips, NumSamples);
+}
+
+bool FOpenXRHMD::AllocateDepthSwapChain(uint32 Index, uint32 SizeX, uint32 SizeY, uint32 NumMips, uint32 NumSamples)
+{
+	check(IsInRenderingThread());
+
+	// @todo: Format enumeration can be moved to a common place.
+	const FormatMap* DepthSwapchainFormat = nullptr;
+	uint32 SwapChainFormatsCount = 0;
+	XR_ENSURE(xrEnumerateSwapchainFormats(Session, 0, &SwapChainFormatsCount, nullptr));
+
+	TArray<int64> Formats;
+	Formats.SetNum(SwapChainFormatsCount);
+	XR_ENSURE(xrEnumerateSwapchainFormats(Session, Formats.Num(), &SwapChainFormatsCount, Formats.GetData()));
+	ensure(SwapChainFormatsCount == Formats.Num());
+
+	for (int i = 0; i < _countof(SupportedDepthSwapchainFormats); ++i)
+	{
+		if (Formats.Contains(SupportedDepthSwapchainFormats[i].DxFormat))
+		{
+			DepthSwapchainFormat = &SupportedDepthSwapchainFormats[i];
+			break;
+		}
+	}
+
+	if (DepthSwapchainFormat == nullptr)
+	{
+		UE_LOG(LogHMD, Log, TEXT("No valid depth swapchain format found."));
+		return false;
+	}
+
+	XrSwapchain SwapchainHandle;
+	XrSwapchainCreateInfo info;
+	info.type = XR_TYPE_SWAPCHAIN_CREATE_INFO;
+	info.next = nullptr;
+	info.createFlags = 0;
+	info.usageFlags = XR_SWAPCHAIN_USAGE_DEPTH_STENCIL_ATTACHMENT_BIT | XR_SWAPCHAIN_USAGE_SAMPLED_BIT;
+	info.format = DepthSwapchainFormat->DxFormat; // FIXME: (DXGI_FORMAT)GPixelFormats[Format].PlatformFormat;
+	info.sampleCount = NumSamples;
+	info.width = SizeX;
+	info.height = SizeY;
+	info.faceCount = 1;
+	info.arraySize = 1;
+	info.mipCount = NumMips;
+	if (!XR_ENSURE(xrCreateSwapchain(Session, &info, &SwapchainHandle)))
+	{
+		return false;
+	}
+
+	uint32_t ChainCount;
+	TArray<XrSwapchainImageD3D11KHR> Images;
+	xrEnumerateSwapchainImages(SwapchainHandle, 0, &ChainCount, nullptr);
+
+	Images.AddZeroed(ChainCount);
+	for (auto& Image : Images)
+	{
+		Image.type = XR_TYPE_SWAPCHAIN_IMAGE_D3D11_KHR;
+	}
+	XR_ENSURE(xrEnumerateSwapchainImages(SwapchainHandle, ChainCount, &ChainCount, reinterpret_cast<XrSwapchainImageBaseHeader*>(Images.GetData())));
+
+	FD3D11DynamicRHI* DynamicRHI = static_cast<FD3D11DynamicRHI*>(GDynamicRHI);
+	TArray<FTextureRHIRef> TextureChain;
+	// @todo: Once things settle down, the chain target will be created below in the CreateXRSwapChain call, via an RHI "CreateAliasedTexture" call.
+	FTextureRHIRef ChainTarget = static_cast<FTextureRHIRef>(DynamicRHI->RHICreateTexture2DFromResource(DepthSwapchainFormat->PixelFormat, TexCreate_DepthStencilTargetable | TexCreate_ShaderResource, FClearValueBinding::DepthFar, Images[0].texture));
+	for (const auto& Image : Images)
+	{
+		TextureChain.Add(static_cast<FTextureRHIRef>(DynamicRHI->RHICreateTexture2DFromResource(DepthSwapchainFormat->PixelFormat, TexCreate_DepthStencilTargetable | TexCreate_ShaderResource, FClearValueBinding::DepthFar, Image.texture)));
+	}
+
+	DepthSwapChain = CreateXRSwapChain<FOpenXRSwapchain>(MoveTemp(TextureChain), ChainTarget, SwapchainHandle);
+
+	return true;
+}
+
+bool FOpenXRHMD::AllocateDepthTexture(uint32 Index, uint32 SizeX, uint32 SizeY, uint8 Format, uint32 NumMips, uint32 InTexFlags, uint32 TargetableTextureFlags, FTexture2DRHIRef& OutTargetableTexture, FTexture2DRHIRef& OutShaderResourceTexture, uint32 NumSamples)
+{
+	if (DepthSwapChain == nullptr || DepthSwapChain->GetTexture2D()->GetSizeX() != SizeX || DepthSwapChain->GetTexture2D()->GetSizeY() != SizeY)
+	{
+		return false;
+	}
+
+	OutTargetableTexture = OutShaderResourceTexture = (FTexture2DRHIRef&)DepthSwapChain->GetTextureRef();
+
 	return true;
 }
 
@@ -806,9 +896,16 @@ void FOpenXRHMD::OnBeginRendering_RenderThread(FRHICommandListImmediate& RHICmdL
 	BaseTransform = FTransform(MainView->BaseHmdOrientation, MainView->BaseHmdLocation);
 
 	Swapchain->IncrementSwapChainIndex_RHIThread(FrameStateRHI.predictedDisplayPeriod);
+	DepthSwapChain->IncrementSwapChainIndex_RHIThread(FrameStateRHI.predictedDisplayPeriod);
 
 	ViewsRHI.SetNum(Views.Num());
+	DepthLayersRHI.SetNum(Views.Num());
+
 	int32 OffsetX = 0;
+
+	const float WorldScale = GetWorldToMetersScale() * (1.0 / 100.0f); // physical scale is 100 UUs/meter
+	float NearZ = GNearClippingPlane * WorldScale;
+
 	for (int32 ViewIndex = 0; ViewIndex < Views.Num(); ViewIndex++)
 	{
 		const XrView& View = Views[ViewIndex];
@@ -816,8 +913,10 @@ void FOpenXRHMD::OnBeginRendering_RenderThread(FRHICommandListImmediate& RHICmdL
 		FTransform ViewTransform = ToFTransform(View.pose, GetWorldToMetersScale());
 
 		XrCompositionLayerProjectionView& Projection = ViewsRHI[ViewIndex];
+		XrCompositionLayerDepthInfoKHR& DepthLayer = DepthLayersRHI[ViewIndex];
+
 		Projection.type = XR_TYPE_COMPOSITION_LAYER_PROJECTION_VIEW;
-		Projection.next = nullptr;
+		Projection.next = &DepthLayer;
 		Projection.fov = View.fov;
 		Projection.pose = ToXrPose(ViewTransform * BaseTransform, GetWorldToMetersScale());
 		Projection.subImage.swapchain = static_cast<FOpenXRSwapchain*>(GetSwapchain())->GetHandle();
@@ -829,6 +928,17 @@ void FOpenXRHMD::OnBeginRendering_RenderThread(FRHICommandListImmediate& RHICmdL
 				(int32)Config.recommendedImageRectHeight
 			}
 		};
+
+		DepthLayer.type = XR_TYPE_COMPOSITION_LAYER_DEPTH_INFO_KHR;
+		DepthLayer.next = nullptr;
+		DepthLayer.subImage.swapchain = static_cast<FOpenXRSwapchain*>(GetDepthSwapchain())->GetHandle();
+		DepthLayer.subImage.imageArrayIndex = 0;
+		DepthLayer.subImage.imageRect = Projection.subImage.imageRect;
+		DepthLayer.minDepth = 1.0f;
+		DepthLayer.maxDepth = 0.0f;
+		DepthLayer.nearZ = NearZ;
+		DepthLayer.farZ = FLT_MAX;
+
 		OffsetX += Config.recommendedImageRectWidth;
 	}
 
@@ -938,6 +1048,7 @@ void FOpenXRHMD::FinishRendering()
 	Layer.views = ViewsRHI.GetData();
 
 	Swapchain->ReleaseCurrentImage_RHIThread();
+	DepthSwapChain->ReleaseCurrentImage_RHIThread();
 
 	XrFrameEndInfo EndInfo;
 	XrCompositionLayerBaseHeader* Headers[1] = { reinterpret_cast<XrCompositionLayerBaseHeader*>(&Layer) };
diff --git a/Engine/Plugins/Runtime/OpenXR/Source/OpenXRHMD/Private/OpenXRHMD.h b/Engine/Plugins/Runtime/OpenXR/Source/OpenXRHMD/Private/OpenXRHMD.h
--- a/Engine/Plugins/Runtime/OpenXR/Source/OpenXRHMD/Private/OpenXRHMD.h
+++ b/Engine/Plugins/Runtime/OpenXR/Source/OpenXRHMD/Private/OpenXRHMD.h
@@ -96,6 +96,8 @@ protected:
 	/** FXRTrackingSystemBase protected interface */
 	virtual float GetWorldToMetersScale() const override;
 
+	bool AllocateDepthSwapChain(uint32 Index, uint32 SizeX, uint32 SizeY, uint32 NumMips, uint32 NumSamples);
+
 public:
 	/** IHeadMountedDisplay interface */
 	virtual bool IsHMDConnected() override { return true; }
@@ -140,6 +142,9 @@ public:
 	/** IStereoRenderTargetManager */
 	virtual bool ShouldUseSeparateRenderTarget() const override { return true; }
 	virtual bool AllocateRenderTargetTexture(uint32 Index, uint32 SizeX, uint32 SizeY, uint8 Format, uint32 NumMips, uint32 Flags, uint32 TargetableTextureFlags, FTexture2DRHIRef& OutTargetableTexture, FTexture2DRHIRef& OutShaderResourceTexture, uint32 NumSamples = 1) override;
+	virtual bool NeedReAllocateDepthTexture(const TRefCountPtr<IPooledRenderTarget>& DepthTarget) override final { return false; }
+	virtual bool AllocateDepthTexture(uint32 Index, uint32 SizeX, uint32 SizeY, uint8 Format, uint32 NumMips, uint32 InTexFlags, uint32 TargetableTextureFlags, FTexture2DRHIRef& OutTargetableTexture, FTexture2DRHIRef& OutShaderResourceTexture, uint32 NumSamples = 1) override final;
+
 	virtual FXRRenderBridge* GetActiveRenderBridge_GameThread(bool bUseSeparateRenderTarget) override;
 
 	/** IStereoRenderTargetManager */
@@ -161,6 +166,7 @@ public:
 	OPENXRHMD_API int32 AddActionDevice(XrAction Action);
 
 	FXRSwapChain* GetSwapchain() { return Swapchain.Get(); }
+	FXRSwapChain* GetDepthSwapchain() { return DepthSwapChain.Get(); }
 	XrInstance GetInstance() { return Instance; }
 	XrSystemId GetSystem() { return System; }
 	XrSession GetSession() { return Session; }
@@ -189,6 +195,7 @@ private:
 	TArray<XrViewConfigurationView> Configs;
 	TArray<XrView>			Views;
 	TArray<XrCompositionLayerProjectionView> ViewsRHI;
+	TArray<XrCompositionLayerDepthInfoKHR> DepthLayersRHI;
 #if 0
 	TArray<FHMDViewMesh>	HiddenAreaMeshes;
 #endif
@@ -197,4 +204,5 @@ private:
 	IRendererModule*		RendererModule;
 
 	FXRSwapChainPtr			Swapchain;
+	FXRSwapChainPtr			DepthSwapChain;
 };
