Deprecate MotionController keys

- Allow keys to be flagged as "bindable in blueprints" and "bindable to actions" separately
- Mark all keys intended for XR projects as not bindable in blueprints
- Mark MotionController keys as not bindable to actions as well
- Replace the SteamVRController module with SteamVRInput based on the marketplace plugin

diff --git a/Engine/Config/BaseInput.ini b/Engine/Config/BaseInput.ini
--- a/Engine/Config/BaseInput.ini
+++ b/Engine/Config/BaseInput.ini
@@ -240,85 +240,3 @@ HistoryColor=(R=180,G=180,B=180,A=255)
 AutoCompleteCommandColor=(B=185,G=109,R=144,A=255)
 AutoCompleteCVarColor=(B=86,G=158,R=86,A=255)
 AutoCompleteFadedColor=(R=100,G=100,B=100,A=255)
-
-[/Script/OpenXRInput.XRInputSettings]
-SimpleController=(Supported=False,ControllerPose=GripPose)
-DaydreamController=(Supported=False,ControllerPose=GripPose)
-ViveController=(Supported=True,ControllerPose=GripPose)
-+ViveBindings=(Path="/user/hand/left/input/menu/click",Key=MotionController_Left_Shoulder)
-+ViveBindings=(Path="/user/hand/right/input/menu/click",Key=MotionController_Right_Shoulder)
-+ViveBindings=(Path="/user/hand/left/input/squeeze",Key=MotionController_Left_Grip1)
-+ViveBindings=(Path="/user/hand/right/input/squeeze",Key=MotionController_Right_Grip1)
-+ViveBindings=(Path="/user/hand/left/input/squeeze",Key=MotionController_Left_Grip1Axis)
-+ViveBindings=(Path="/user/hand/right/input/squeeze",Key=MotionController_Right_Grip1Axis)
-+ViveBindings=(Path="/user/hand/left/input/trigger",Key=MotionController_Left_Trigger)
-+ViveBindings=(Path="/user/hand/right/input/trigger",Key=MotionController_Right_Trigger)
-+ViveBindings=(Path="/user/hand/left/input/trigger",Key=MotionController_Left_TriggerAxis)
-+ViveBindings=(Path="/user/hand/right/input/trigger",Key=MotionController_Right_TriggerAxis)
-+ViveBindings=(Path="/user/hand/left/input/trackpad/x",Key=MotionController_Left_Thumbstick_X)
-+ViveBindings=(Path="/user/hand/right/input/trackpad/x",Key=MotionController_Right_Thumbstick_X)
-+ViveBindings=(Path="/user/hand/left/input/trackpad/y",Key=MotionController_Left_Thumbstick_Y)
-+ViveBindings=(Path="/user/hand/right/input/trackpad/y",Key=MotionController_Right_Thumbstick_Y)
-+ViveBindings=(Path="/user/hand/left/input/trackpad/click",Key=MotionController_Left_Thumbstick)
-+ViveBindings=(Path="/user/hand/right/input/trackpad/click",Key=MotionController_Right_Thumbstick)
-ViveProHeadset=(Supported=False,ControllerPose=GripPose)
-MixedRealityController=(Supported=True,ControllerPose=GripPose)
-+MixedRealityBindings=(Path="/user/hand/left/input/menu/click",Key=MotionController_Left_Shoulder)
-+MixedRealityBindings=(Path="/user/hand/right/input/menu/click",Key=MotionController_Right_Shoulder)
-+MixedRealityBindings=(Path="/user/hand/left/input/squeeze",Key=MotionController_Left_Grip1)
-+MixedRealityBindings=(Path="/user/hand/right/input/squeeze",Key=MotionController_Right_Grip1)
-+MixedRealityBindings=(Path="/user/hand/left/input/squeeze",Key=MotionController_Left_Grip1Axis)
-+MixedRealityBindings=(Path="/user/hand/right/input/squeeze",Key=MotionController_Right_Grip1Axis)
-+MixedRealityBindings=(Path="/user/hand/left/input/trigger",Key=MotionController_Left_Trigger)
-+MixedRealityBindings=(Path="/user/hand/right/input/trigger",Key=MotionController_Right_Trigger)
-+MixedRealityBindings=(Path="/user/hand/left/input/trigger",Key=MotionController_Left_TriggerAxis)
-+MixedRealityBindings=(Path="/user/hand/right/input/trigger",Key=MotionController_Right_TriggerAxis)
-+MixedRealityBindings=(Path="/user/hand/left/input/thumbstick/x",Key=MotionController_Left_Thumbstick_X)
-+MixedRealityBindings=(Path="/user/hand/right/input/thumbstick/x",Key=MotionController_Right_Thumbstick_X)
-+MixedRealityBindings=(Path="/user/hand/left/input/thumbstick/y",Key=MotionController_Left_Thumbstick_Y)
-+MixedRealityBindings=(Path="/user/hand/right/input/thumbstick/y",Key=MotionController_Right_Thumbstick_Y)
-+MixedRealityBindings=(Path="/user/hand/left/input/thumbstick/click",Key=MotionController_Left_Thumbstick)
-+MixedRealityBindings=(Path="/user/hand/right/input/thumbstick/click",Key=MotionController_Right_Thumbstick)
-OculusGoController=(Supported=False,ControllerPose=GripPose)
-OculusTouchController=(Supported=True,ControllerPose=GripPose)
-+OculusTouchBindings=(Path="/user/hand/left/input/menu/click",Key=MotionController_Left_Shoulder)
-+OculusTouchBindings=(Path="/user/hand/right/input/menu/click",Key=MotionController_Right_Shoulder)
-+OculusTouchBindings=(Path="/user/hand/left/input/squeeze",Key=MotionController_Left_Grip1)
-+OculusTouchBindings=(Path="/user/hand/right/input/squeeze",Key=MotionController_Right_Grip1)
-+OculusTouchBindings=(Path="/user/hand/left/input/squeeze",Key=MotionController_Left_Grip1Axis)
-+OculusTouchBindings=(Path="/user/hand/right/input/squeeze",Key=MotionController_Right_Grip1Axis)
-+OculusTouchBindings=(Path="/user/hand/left/input/trigger",Key=MotionController_Left_Trigger)
-+OculusTouchBindings=(Path="/user/hand/right/input/trigger",Key=MotionController_Right_Trigger)
-+OculusTouchBindings=(Path="/user/hand/left/input/trigger",Key=MotionController_Left_TriggerAxis)
-+OculusTouchBindings=(Path="/user/hand/right/input/trigger",Key=MotionController_Right_TriggerAxis)
-+OculusTouchBindings=(Path="/user/hand/left/input/thumbstick/x",Key=MotionController_Left_Thumbstick_X)
-+OculusTouchBindings=(Path="/user/hand/right/input/thumbstick/x",Key=MotionController_Right_Thumbstick_X)
-+OculusTouchBindings=(Path="/user/hand/left/input/thumbstick/y",Key=MotionController_Left_Thumbstick_Y)
-+OculusTouchBindings=(Path="/user/hand/right/input/thumbstick/y",Key=MotionController_Right_Thumbstick_Y)
-+OculusTouchBindings=(Path="/user/hand/left/input/thumbstick/click",Key=MotionController_Left_Thumbstick)
-+OculusTouchBindings=(Path="/user/hand/right/input/thumbstick/click",Key=MotionController_Right_Thumbstick)
-+OculusTouchBindings=(Path="/user/hand/left/input/x/click",Key=MotionController_Left_FaceButton1)
-+OculusTouchBindings=(Path="/user/hand/right/input/a/click",Key=MotionController_Right_FaceButton1)
-+OculusTouchBindings=(Path="/user/hand/left/input/y/click",Key=MotionController_Left_FaceButton2)
-+OculusTouchBindings=(Path="/user/hand/right/input/b/click",Key=MotionController_Right_FaceButton2)
-ValveIndexController=(Supported=True,ControllerPose=GripPose)
-+ValveIndexBindings=(Path="/user/hand/left/input/menu/click",Key=MotionController_Left_Shoulder)
-+ValveIndexBindings=(Path="/user/hand/right/input/menu/click",Key=MotionController_Right_Shoulder)
-+ValveIndexBindings=(Path="/user/hand/left/input/squeeze",Key=MotionController_Left_Grip1)
-+ValveIndexBindings=(Path="/user/hand/right/input/squeeze",Key=MotionController_Right_Grip1)
-+ValveIndexBindings=(Path="/user/hand/left/input/squeeze",Key=MotionController_Left_Grip1Axis)
-+ValveIndexBindings=(Path="/user/hand/right/input/squeeze",Key=MotionController_Right_Grip1Axis)
-+ValveIndexBindings=(Path="/user/hand/left/input/trigger",Key=MotionController_Left_Trigger)
-+ValveIndexBindings=(Path="/user/hand/right/input/trigger",Key=MotionController_Right_Trigger)
-+ValveIndexBindings=(Path="/user/hand/left/input/trigger",Key=MotionController_Left_TriggerAxis)
-+ValveIndexBindings=(Path="/user/hand/right/input/trigger",Key=MotionController_Right_TriggerAxis)
-+ValveIndexBindings=(Path="/user/hand/left/input/thumbstick/x",Key=MotionController_Left_Thumbstick_X)
-+ValveIndexBindings=(Path="/user/hand/right/input/thumbstick/x",Key=MotionController_Right_Thumbstick_X)
-+ValveIndexBindings=(Path="/user/hand/left/input/thumbstick/y",Key=MotionController_Left_Thumbstick_Y)
-+ValveIndexBindings=(Path="/user/hand/right/input/thumbstick/y",Key=MotionController_Right_Thumbstick_Y)
-+ValveIndexBindings=(Path="/user/hand/left/input/thumbstick/click",Key=MotionController_Left_Thumbstick)
-+ValveIndexBindings=(Path="/user/hand/right/input/thumbstick/click",Key=MotionController_Right_Thumbstick)
-+ValveIndexBindings=(Path="/user/hand/left/input/a/click",Key=MotionController_Left_FaceButton1)
-+ValveIndexBindings=(Path="/user/hand/right/input/a/click",Key=MotionController_Right_FaceButton1)
-+ValveIndexBindings=(Path="/user/hand/left/input/b/click",Key=MotionController_Left_FaceButton2)
-+ValveIndexBindings=(Path="/user/hand/right/input/b/click",Key=MotionController_Right_FaceButton2)
diff --git a/Engine/Plugins/Lumin/MagicLeap/Source/MagicLeapController/Private/MagicLeapController.cpp b/Engine/Plugins/Lumin/MagicLeap/Source/MagicLeapController/Private/MagicLeapController.cpp
--- a/Engine/Plugins/Lumin/MagicLeap/Source/MagicLeapController/Private/MagicLeapController.cpp
+++ b/Engine/Plugins/Lumin/MagicLeap/Source/MagicLeapController/Private/MagicLeapController.cpp
@@ -523,6 +523,7 @@ void FMagicLeapController::InitializeInputCallbacks()
 			if (ControllerHand != EControllerHand::ControllerHand_Count)
 			{
 				Controller->PendingButtonEvents.Enqueue(MakeTuple(MLToUnrealButton(ControllerHand, button), true));
+				Controller->PendingButtonEvents.Enqueue(MakeTuple(MLToUnrealButtonLegacy(ControllerHand, button), true));
 			}
 		}
 	};
@@ -536,6 +537,7 @@ void FMagicLeapController::InitializeInputCallbacks()
 			if (ControllerHand != EControllerHand::ControllerHand_Count)
 			{
 				Controller->PendingButtonEvents.Enqueue(MakeTuple(MLToUnrealButton(ControllerHand, button), false));
+				Controller->PendingButtonEvents.Enqueue(MakeTuple(MLToUnrealButtonLegacy(ControllerHand, button), false));
 			}
 		}
 	};
@@ -738,8 +740,8 @@ void FMagicLeapController::CreateEntityTracker()
 	Result = MLInputCreate(&InputConfig, &InputTracker);
 	if (Result != MLResult_Ok)
 	{
-		UE_LOG(LogMagicLeapController, Error, 
-			TEXT("MLInputCreate failed with error %s."), 
+		UE_LOG(LogMagicLeapController, Error,
+			TEXT("MLInputCreate failed with error %s."),
 			UTF8_TO_TCHAR(MLGetResultString(Result)));
 	}
 #endif //PLATFORM_LUMIN
@@ -771,8 +773,8 @@ void FMagicLeapController::DestroyEntityTracker()
 		MLResult Result = MLInputDestroy(InputTracker);
 		if (Result != MLResult_Ok)
 		{
-			UE_LOG(LogMagicLeapController, Error, 
-				TEXT("MLInputDestroy failed with error %s!"), 
+			UE_LOG(LogMagicLeapController, Error,
+				TEXT("MLInputDestroy failed with error %s!"),
 				UTF8_TO_TCHAR(MLGetResultString(Result)));
 		}
 	}
@@ -948,11 +950,23 @@ void FMagicLeapController::UpdateControllerStateFromInputTracker(const IMagicLea
 			// Touch0 activate/deactivate
 			if (CurrControllerState.bTouchActive[0] && !PrevControllerState.bTouchActive[0])
 			{
-				PendingButtonEvents.Enqueue(MakeTuple(MLTouchToUnrealThumbstickButton(Hand), true));
+				PendingButtonEvents.Enqueue(MakeTuple(MLTouchToUnrealTrackpadButton(Hand), true));
+				PendingButtonEvents.Enqueue(MakeTuple(MLTouchToUnrealTrackpadButtonLegacy(Hand), true));
 			}
 			else if (PrevControllerState.bTouchActive[0] && !CurrControllerState.bTouchActive[0])
 			{
-				PendingButtonEvents.Enqueue(MakeTuple(MLTouchToUnrealThumbstickButton(Hand), false));
+				PendingButtonEvents.Enqueue(MakeTuple(MLTouchToUnrealTrackpadButton(Hand), false));
+				PendingButtonEvents.Enqueue(MakeTuple(MLTouchToUnrealTrackpadButtonLegacy(Hand), false));
+			}
+
+			// Touch1 activate/deactivate
+			if (CurrControllerState.bTouchActive[1] && !PrevControllerState.bTouchActive[1])
+			{
+				PendingButtonEvents.Enqueue(MakeTuple(MLTouchToUnrealTouch1Button(Hand), true));
+			}
+			else if (PrevControllerState.bTouchActive[1] && !CurrControllerState.bTouchActive[1])
+			{
+				PendingButtonEvents.Enqueue(MakeTuple(MLTouchToUnrealTouch1Button(Hand), false));
 			}
 
 			// Convert trigger value to trigger press/release events
@@ -964,11 +978,13 @@ void FMagicLeapController::UpdateControllerStateFromInputTracker(const IMagicLea
 			if (IsTriggerKeyPressing)
 			{
 				PendingButtonEvents.Enqueue(MakeTuple(MLTriggerToUnrealTriggerKey(Hand), true));
+				PendingButtonEvents.Enqueue(MakeTuple(MLTriggerToUnrealTriggerKeyLegacy(Hand), true));
 				CurrControllerState.bTriggerKeyPressing = true;
 			}
 			else if (IsTriggerKeyReleasing)
 			{
 				PendingButtonEvents.Enqueue(MakeTuple(MLTriggerToUnrealTriggerKey(Hand), false));
+				PendingButtonEvents.Enqueue(MakeTuple(MLTriggerToUnrealTriggerKeyLegacy(Hand), false));
 				CurrControllerState.bTriggerKeyPressing = false;
 			}
 		}
@@ -1146,11 +1162,35 @@ void FMagicLeapController::UnregisterTouchpadGestureReceiver(IMagicLeapTouchpadG
 
 void FMagicLeapController::AddKeys()
 {
-	EKeys::AddKey(FKeyDetails(FMagicLeapKeys::MotionController_Left_Thumbstick_Z, LOCTEXT("MotionController_Left_Thumbstick_Z", "MotionController (L) Thumbstick Z"), FKeyDetails::GamepadKey | FKeyDetails::FloatAxis));
-	EKeys::AddKey(FKeyDetails(FMagicLeapKeys::Left_HomeButton, LOCTEXT("MagicLeap_Left_HomeButton", "ML (L) Home Button"), FKeyDetails::GamepadKey));
-
-	EKeys::AddKey(FKeyDetails(FMagicLeapKeys::MotionController_Right_Thumbstick_Z, LOCTEXT("MotionController_Right_Thumbstick_Z", "MotionController (R) Thumbstick Z"), FKeyDetails::GamepadKey | FKeyDetails::FloatAxis));
-	EKeys::AddKey(FKeyDetails(FMagicLeapKeys::Right_HomeButton, LOCTEXT("MagicLeap_Right_HomeButton", "ML (R) Home Button"), FKeyDetails::GamepadKey));
+	EKeys::AddMenuCategoryDisplayInfo("MagicLeap", LOCTEXT("MagicLeapSubCategory", "Magic Leap"), TEXT("GraphEditor.PadEvent_16x"));
+
+	EKeys::AddKey(FKeyDetails(FMagicLeapKeys::MotionController_Left_Thumbstick_Z, LOCTEXT("MotionController_Left_Thumbstick_Z", "MotionController (L) Thumbstick Z"), FKeyDetails::GamepadKey | FKeyDetails::FloatAxis | FKeyDetails::NotBlueprintBindableKey | FKeyDetails::NotActionBindableKey));
+	EKeys::AddKey(FKeyDetails(FMagicLeapKeys::Left_HomeButton, LOCTEXT("MagicLeap_Left_HomeButton", "ML (L) Home Button"), FKeyDetails::GamepadKey | FKeyDetails::NotBlueprintBindableKey, "MagicLeap"));
+	EKeys::AddKey(FKeyDetails(FMagicLeapKeys::Left_Bumper, LOCTEXT("MagicLeap_Left_Bumper", "ML (L) Bumper"), FKeyDetails::GamepadKey | FKeyDetails::NotBlueprintBindableKey, "MagicLeap"));
+	EKeys::AddKey(FKeyDetails(FMagicLeapKeys::Left_Trigger, LOCTEXT("MagicLeap_Left_Trigger", "ML (L) Trigger"), FKeyDetails::GamepadKey | FKeyDetails::NotBlueprintBindableKey, "MagicLeap"));
+	EKeys::AddKey(FKeyDetails(FMagicLeapKeys::Left_Trigger_Axis, LOCTEXT("MagicLeap_Left_Trigger_Axis", "ML (L) Trigger Axis"), FKeyDetails::GamepadKey | FKeyDetails::FloatAxis | FKeyDetails::NotBlueprintBindableKey, "MagicLeap"));
+	EKeys::AddKey(FKeyDetails(FMagicLeapKeys::Left_Trackpad_X, LOCTEXT("MagicLeap_Left_Trackpad_X", "ML (L) Trackpad X"), FKeyDetails::GamepadKey | FKeyDetails::FloatAxis | FKeyDetails::NotBlueprintBindableKey, "MagicLeap"));
+	EKeys::AddKey(FKeyDetails(FMagicLeapKeys::Left_Trackpad_Y, LOCTEXT("MagicLeap_Left_Trackpad_Y", "ML (L) Trackpad Y"), FKeyDetails::GamepadKey | FKeyDetails::FloatAxis | FKeyDetails::NotBlueprintBindableKey, "MagicLeap"));
+	EKeys::AddKey(FKeyDetails(FMagicLeapKeys::Left_Trackpad_Force, LOCTEXT("MagicLeap_Left_Trackpad_Force", "ML (L) Trackpad Force"), FKeyDetails::GamepadKey | FKeyDetails::FloatAxis | FKeyDetails::NotBlueprintBindableKey, "MagicLeap"));
+	EKeys::AddKey(FKeyDetails(FMagicLeapKeys::Left_Trackpad_Touch, LOCTEXT("MagicLeap_Left_Trackpad_Touch", "ML (L) Trackpad Touch"), FKeyDetails::GamepadKey | FKeyDetails::NotBlueprintBindableKey, "MagicLeap"));
+	EKeys::AddKey(FKeyDetails(FMagicLeapKeys::Left_Touch1_X, LOCTEXT("MagicLeap_Left_Touch1_X", "ML (L) Touch1 X"), FKeyDetails::GamepadKey | FKeyDetails::FloatAxis | FKeyDetails::NotBlueprintBindableKey, "MagicLeap"));
+	EKeys::AddKey(FKeyDetails(FMagicLeapKeys::Left_Touch1_Y, LOCTEXT("MagicLeap_Left_Touch1_Y", "ML (L) Touch1 Y"), FKeyDetails::GamepadKey | FKeyDetails::FloatAxis | FKeyDetails::NotBlueprintBindableKey, "MagicLeap"));
+	EKeys::AddKey(FKeyDetails(FMagicLeapKeys::Left_Touch1_Force, LOCTEXT("MagicLeap_Left_Touch1_Force", "ML (L) Touch1 Force"), FKeyDetails::GamepadKey | FKeyDetails::FloatAxis | FKeyDetails::NotBlueprintBindableKey, "MagicLeap"));
+	EKeys::AddKey(FKeyDetails(FMagicLeapKeys::Left_Touch1_Touch, LOCTEXT("MagicLeap_Left_Touch1_Touch", "ML (L) Touch1 Touch"), FKeyDetails::GamepadKey | FKeyDetails::NotBlueprintBindableKey, "MagicLeap"));
+
+	EKeys::AddKey(FKeyDetails(FMagicLeapKeys::MotionController_Right_Thumbstick_Z, LOCTEXT("MotionController_Right_Thumbstick_Z", "MotionController (R) Thumbstick Z"), FKeyDetails::GamepadKey | FKeyDetails::FloatAxis | FKeyDetails::NotBlueprintBindableKey | FKeyDetails::NotActionBindableKey));
+	EKeys::AddKey(FKeyDetails(FMagicLeapKeys::Right_HomeButton, LOCTEXT("MagicLeap_Right_HomeButton", "ML (R) Home Button"), FKeyDetails::GamepadKey | FKeyDetails::NotBlueprintBindableKey, "MagicLeap"));
+	EKeys::AddKey(FKeyDetails(FMagicLeapKeys::Right_Bumper, LOCTEXT("MagicLeap_Right_Bumper", "ML (R) Bumper"), FKeyDetails::GamepadKey | FKeyDetails::NotBlueprintBindableKey, "MagicLeap"));
+	EKeys::AddKey(FKeyDetails(FMagicLeapKeys::Right_Trigger, LOCTEXT("MagicLeap_Right_Trigger", "ML (R) Trigger"), FKeyDetails::GamepadKey | FKeyDetails::NotBlueprintBindableKey, "MagicLeap"));
+	EKeys::AddKey(FKeyDetails(FMagicLeapKeys::Right_Trigger_Axis, LOCTEXT("MagicLeap_Right_Trigger_Axis", "ML (R) Trigger Axis"), FKeyDetails::GamepadKey | FKeyDetails::FloatAxis | FKeyDetails::NotBlueprintBindableKey, "MagicLeap"));
+	EKeys::AddKey(FKeyDetails(FMagicLeapKeys::Right_Trackpad_X, LOCTEXT("MagicLeap_Right_Trackpad_X", "ML (R) Trackpad X"), FKeyDetails::GamepadKey | FKeyDetails::FloatAxis | FKeyDetails::NotBlueprintBindableKey, "MagicLeap"));
+	EKeys::AddKey(FKeyDetails(FMagicLeapKeys::Right_Trackpad_Y, LOCTEXT("MagicLeap_Right_Trackpad_Y", "ML (R) Trackpad Y"), FKeyDetails::GamepadKey | FKeyDetails::FloatAxis | FKeyDetails::NotBlueprintBindableKey, "MagicLeap"));
+	EKeys::AddKey(FKeyDetails(FMagicLeapKeys::Right_Trackpad_Force, LOCTEXT("MagicLeap_Right_Trackpad_Force", "ML (R) Trackpad Force"), FKeyDetails::GamepadKey | FKeyDetails::FloatAxis | FKeyDetails::NotBlueprintBindableKey, "MagicLeap"));
+	EKeys::AddKey(FKeyDetails(FMagicLeapKeys::Right_Trackpad_Touch, LOCTEXT("MagicLeap_Right_Trackpad_Touch", "ML (R) Trackpad Touch"), FKeyDetails::GamepadKey | FKeyDetails::NotBlueprintBindableKey, "MagicLeap"));
+	EKeys::AddKey(FKeyDetails(FMagicLeapKeys::Right_Touch1_X, LOCTEXT("MagicLeap_Right_Touch1_X", "ML (R) Touch1 X"), FKeyDetails::GamepadKey | FKeyDetails::FloatAxis | FKeyDetails::NotBlueprintBindableKey, "MagicLeap"));
+	EKeys::AddKey(FKeyDetails(FMagicLeapKeys::Right_Touch1_Y, LOCTEXT("MagicLeap_Right_Touch1_Y", "ML (R) Touch1 Y"), FKeyDetails::GamepadKey | FKeyDetails::FloatAxis | FKeyDetails::NotBlueprintBindableKey, "MagicLeap"));
+	EKeys::AddKey(FKeyDetails(FMagicLeapKeys::Right_Touch1_Force, LOCTEXT("MagicLeap_Right_Touch1_Force", "ML (R) Touch1 Force"), FKeyDetails::GamepadKey | FKeyDetails::FloatAxis | FKeyDetails::NotBlueprintBindableKey, "MagicLeap"));
+	EKeys::AddKey(FKeyDetails(FMagicLeapKeys::Right_Touch1_Touch, LOCTEXT("MagicLeap_Right_Touch1_Touch", "ML (R) Touch1 Touch"), FKeyDetails::GamepadKey | FKeyDetails::NotBlueprintBindableKey, "MagicLeap"));
 }
 
 bool FMagicLeapController::PlayLEDPattern(FName MotionSource, EMagicLeapControllerLEDPattern LEDPattern, EMagicLeapControllerLEDColor LEDColor, float DurationInSec)
@@ -1288,19 +1328,26 @@ void FMagicLeapController::SendControllerEventsForHand(EControllerHand Hand)
 		checkf(PrevControllerState != nullptr, TEXT("Unpossible"));
 
 		// Analog touch coords
-		// Touch 0 maps to Motion Controller Thumbstick for hand
-		// Touch 1 is currently not available (we have nothing to map it to)
+		// Touch 0 maps to ML Trackpad
+		// Touch 1 maps to ML Touch1
 		if (CurrControllerState->bTouchActive[0])
 		{
 			MagicLeap::EnableInput EnableInputFromHMD;
 			// fixes unreferenced parameter error for Windows package builds.
 			(void)EnableInputFromHMD;
 
-			MessageHandler->OnControllerAnalog(MLTouchToUnrealThumbstickAxis(Hand, 0),
+			MessageHandler->OnControllerAnalog(MLTouchToUnrealTrackpadAxis(Hand, 0),
+				DeviceIndex, CurrControllerState->TouchPosAndForce[0].X);
+			MessageHandler->OnControllerAnalog(MLTouchToUnrealTrackpadAxis(Hand, 1),
+				DeviceIndex, CurrControllerState->TouchPosAndForce[0].Y);
+			MessageHandler->OnControllerAnalog(MLTouchToUnrealTrackpadAxis(Hand, 2),
+				DeviceIndex, CurrControllerState->TouchPosAndForce[0].Z);
+
+			MessageHandler->OnControllerAnalog(MLTouchToUnrealTrackpadAxis(Hand, 0, true),
 				DeviceIndex, CurrControllerState->TouchPosAndForce[0].X);
-			MessageHandler->OnControllerAnalog(MLTouchToUnrealThumbstickAxis(Hand, 1),
+			MessageHandler->OnControllerAnalog(MLTouchToUnrealTrackpadAxis(Hand, 1, true),
 				DeviceIndex, CurrControllerState->TouchPosAndForce[0].Y);
-			MessageHandler->OnControllerAnalog(MLTouchToUnrealThumbstickAxis(Hand, 2),
+			MessageHandler->OnControllerAnalog(MLTouchToUnrealTrackpadAxis(Hand, 2, true),
 				DeviceIndex, CurrControllerState->TouchPosAndForce[0].Z);
 		}
 		else
@@ -1309,9 +1356,40 @@ void FMagicLeapController::SendControllerEventsForHand(EControllerHand Hand)
 			// fixes unreferenced parameter error for Windows package builds.
 			(void)EnableInputFromHMD;
 
-			MessageHandler->OnControllerAnalog(MLTouchToUnrealThumbstickAxis(Hand, 0), DeviceIndex, 0.0f);
-			MessageHandler->OnControllerAnalog(MLTouchToUnrealThumbstickAxis(Hand, 1), DeviceIndex, 0.0f);
-			MessageHandler->OnControllerAnalog(MLTouchToUnrealThumbstickAxis(Hand, 2), DeviceIndex, 0.0f);
+			MessageHandler->OnControllerAnalog(MLTouchToUnrealTrackpadAxis(Hand, 0), DeviceIndex, 0.0f);
+			MessageHandler->OnControllerAnalog(MLTouchToUnrealTrackpadAxis(Hand, 1), DeviceIndex, 0.0f);
+			MessageHandler->OnControllerAnalog(MLTouchToUnrealTrackpadAxis(Hand, 2), DeviceIndex, 0.0f);
+
+			MessageHandler->OnControllerAnalog(MLTouchToUnrealTrackpadAxis(Hand, 0, true), DeviceIndex, 0.0f);
+			MessageHandler->OnControllerAnalog(MLTouchToUnrealTrackpadAxis(Hand, 1, true), DeviceIndex, 0.0f);
+			MessageHandler->OnControllerAnalog(MLTouchToUnrealTrackpadAxis(Hand, 2, true), DeviceIndex, 0.0f);
+		}
+
+		// Analog touch coords
+		// Touch 0 maps to Motion Controller Thumbstick for hand
+		// Touch 1 is currently not available (we have nothing to map it to)
+		if (CurrControllerState->bTouchActive[1])
+		{
+			MagicLeap::EnableInput EnableInputFromHMD;
+			// fixes unreferenced parameter error for Windows package builds.
+			(void)EnableInputFromHMD;
+
+			MessageHandler->OnControllerAnalog(MLTouchToUnrealTouch1Axis(Hand, 0),
+				DeviceIndex, CurrControllerState->TouchPosAndForce[1].X);
+			MessageHandler->OnControllerAnalog(MLTouchToUnrealTouch1Axis(Hand, 1),
+				DeviceIndex, CurrControllerState->TouchPosAndForce[1].Y);
+			MessageHandler->OnControllerAnalog(MLTouchToUnrealTouch1Axis(Hand, 2),
+				DeviceIndex, CurrControllerState->TouchPosAndForce[1].Z);
+}
+		else
+		{
+			MagicLeap::EnableInput EnableInputFromHMD;
+			// fixes unreferenced parameter error for Windows package builds.
+			(void)EnableInputFromHMD;
+
+			MessageHandler->OnControllerAnalog(MLTouchToUnrealTouch1Axis(Hand, 0), DeviceIndex, 0.0f);
+			MessageHandler->OnControllerAnalog(MLTouchToUnrealTouch1Axis(Hand, 1), DeviceIndex, 0.0f);
+			MessageHandler->OnControllerAnalog(MLTouchToUnrealTouch1Axis(Hand, 2), DeviceIndex, 0.0f);
 		}
 
 		// Analog trigger
@@ -1323,6 +1401,8 @@ void FMagicLeapController::SendControllerEventsForHand(EControllerHand Hand)
 
 			MessageHandler->OnControllerAnalog(MLTriggerToUnrealTriggerAxis(Hand),
 				DeviceIndex, CurrControllerState->TriggerAnalog);
+			MessageHandler->OnControllerAnalogLegacy(MLTriggerToUnrealTriggerAxis(Hand),
+				DeviceIndex, CurrControllerState->TriggerAnalog);
 		}
 	}
 #endif //WITH_MLSDK
diff --git a/Engine/Plugins/Lumin/MagicLeap/Source/MagicLeapController/Private/MagicLeapControllerKeys.cpp b/Engine/Plugins/Lumin/MagicLeap/Source/MagicLeapController/Private/MagicLeapControllerKeys.cpp
--- a/Engine/Plugins/Lumin/MagicLeap/Source/MagicLeapController/Private/MagicLeapControllerKeys.cpp
+++ b/Engine/Plugins/Lumin/MagicLeap/Source/MagicLeapController/Private/MagicLeapControllerKeys.cpp
@@ -10,12 +10,56 @@ const FName FMagicLeapControllerKeyNames::MotionController_Left_Thumbstick_Z_Nam
 
 const FKey FMagicLeapKeys::Left_HomeButton("MagicLeap_Left_HomeButton");
 const FName FMagicLeapControllerKeyNames::Left_HomeButton_Name("MagicLeap_Left_HomeButton");
+const FKey FMagicLeapKeys::Left_Bumper("MagicLeap_Left_Bumper");
+const FName FMagicLeapControllerKeyNames::Left_Bumper_Name("MagicLeap_Left_Bumper");
+const FKey FMagicLeapKeys::Left_Trigger("MagicLeap_Left_Trigger");
+const FName FMagicLeapControllerKeyNames::Left_Trigger_Name("MagicLeap_Left_Trigger");
+const FKey FMagicLeapKeys::Left_Trigger_Axis("MagicLeap_Left_Trigger_Axis");
+const FName FMagicLeapControllerKeyNames::Left_Trigger_Axis_Name("MagicLeap_Left_Trigger_Axis");
+const FKey FMagicLeapKeys::Left_Trackpad_X("MagicLeap_Left_Trackpad_X");
+const FName FMagicLeapControllerKeyNames::Left_Trackpad_X_Name("MagicLeap_Left_Trackpad_X");
+const FKey FMagicLeapKeys::Left_Trackpad_Y("MagicLeap_Left_Trackpad_Y");
+const FName FMagicLeapControllerKeyNames::Left_Trackpad_Y_Name("MagicLeap_Left_Trackpad_Y");
+const FKey FMagicLeapKeys::Left_Trackpad_Force("MagicLeap_Left_Trackpad_Force");
+const FName FMagicLeapControllerKeyNames::Left_Trackpad_Force_Name("MagicLeap_Left_Trackpad_Force");
+const FKey FMagicLeapKeys::Left_Trackpad_Touch("MagicLeap_Left_Trackpad_Touch");
+const FName FMagicLeapControllerKeyNames::Left_Trackpad_Touch_Name("MagicLeap_Left_Trackpad_Touch");
+const FKey FMagicLeapKeys::Left_Touch1_X("MagicLeap_Left_Touch1_X");
+const FName FMagicLeapControllerKeyNames::Left_Touch1_X_Name("MagicLeap_Left_Touch1_X");
+const FKey FMagicLeapKeys::Left_Touch1_Y("MagicLeap_Left_Touch1_Y");
+const FName FMagicLeapControllerKeyNames::Left_Touch1_Y_Name("MagicLeap_Left_Touch1_Y");
+const FKey FMagicLeapKeys::Left_Touch1_Force("MagicLeap_Left_Touch1_Force");
+const FName FMagicLeapControllerKeyNames::Left_Touch1_Force_Name("MagicLeap_Left_Touch1_Force");
+const FKey FMagicLeapKeys::Left_Touch1_Touch("MagicLeap_Left_Touch1_Touch");
+const FName FMagicLeapControllerKeyNames::Left_Touch1_Touch_Name("MagicLeap_Left_Touch1_Touch");
 
 const FKey FMagicLeapKeys::MotionController_Right_Thumbstick_Z("MotionController_Right_Thumbstick_Z");
 const FName FMagicLeapControllerKeyNames::MotionController_Right_Thumbstick_Z_Name("MotionController_Right_Thumbstick_Z");
 
 const FKey FMagicLeapKeys::Right_HomeButton("MagicLeap_Right_HomeButton");
 const FName FMagicLeapControllerKeyNames::Right_HomeButton_Name("MagicLeap_Right_HomeButton");
+const FKey FMagicLeapKeys::Right_Bumper("MagicLeap_Right_Bumper");
+const FName FMagicLeapControllerKeyNames::Right_Bumper_Name("MagicLeap_Right_Bumper");
+const FKey FMagicLeapKeys::Right_Trigger("MagicLeap_Right_Trigger");
+const FName FMagicLeapControllerKeyNames::Right_Trigger_Name("MagicLeap_Right_Trigger");
+const FKey FMagicLeapKeys::Right_Trigger_Axis("MagicLeap_Right_Trigger_Axis");
+const FName FMagicLeapControllerKeyNames::Right_Trigger_Axis_Name("MagicLeap_Right_Trigger_Axis");
+const FKey FMagicLeapKeys::Right_Trackpad_X("MagicLeap_Right_Trackpad_X");
+const FName FMagicLeapControllerKeyNames::Right_Trackpad_X_Name("MagicLeap_Right_Trackpad_X");
+const FKey FMagicLeapKeys::Right_Trackpad_Y("MagicLeap_Right_Trackpad_Y");
+const FName FMagicLeapControllerKeyNames::Right_Trackpad_Y_Name("MagicLeap_Right_Trackpad_Y");
+const FKey FMagicLeapKeys::Right_Trackpad_Force("MagicLeap_Right_Trackpad_Force");
+const FName FMagicLeapControllerKeyNames::Right_Trackpad_Force_Name("MagicLeap_Right_Trackpad_Force");
+const FKey FMagicLeapKeys::Right_Trackpad_Touch("MagicLeap_Right_Trackpad_Touch");
+const FName FMagicLeapControllerKeyNames::Right_Trackpad_Touch_Name("MagicLeap_Right_Trackpad_Touch");
+const FKey FMagicLeapKeys::Right_Touch1_X("MagicLeap_Right_Touch1_X");
+const FName FMagicLeapControllerKeyNames::Right_Touch1_X_Name("MagicLeap_Right_Touch1_X");
+const FKey FMagicLeapKeys::Right_Touch1_Y("MagicLeap_Right_Touch1_Y");
+const FName FMagicLeapControllerKeyNames::Right_Touch1_Y_Name("MagicLeap_Right_Touch1_Y");
+const FKey FMagicLeapKeys::Right_Touch1_Force("MagicLeap_Right_Touch1_Force");
+const FName FMagicLeapControllerKeyNames::Right_Touch1_Force_Name("MagicLeap_Right_Touch1_Force");
+const FKey FMagicLeapKeys::Right_Touch1_Touch("MagicLeap_Right_Touch1_Touch");
+const FName FMagicLeapControllerKeyNames::Right_Touch1_Touch_Name("MagicLeap_Right_Touch1_Touch");
 
 const FKey FMagicLeapKeys::TouchpadGesture_Swipe_Up("TouchpadGesture_Swipe_Up");
 const FName FMagicLeapControllerKeyNames::TouchpadGesture_Swipe_Up_Name("TouchpadGesture_Swipe_Up");
diff --git a/Engine/Plugins/Lumin/MagicLeap/Source/MagicLeapController/Private/MagicLeapControllerMappings.h b/Engine/Plugins/Lumin/MagicLeap/Source/MagicLeapController/Private/MagicLeapControllerMappings.h
--- a/Engine/Plugins/Lumin/MagicLeap/Source/MagicLeapController/Private/MagicLeapControllerMappings.h
+++ b/Engine/Plugins/Lumin/MagicLeap/Source/MagicLeapController/Private/MagicLeapControllerMappings.h
@@ -58,6 +58,30 @@ const FName& MLToUnrealButton(EControllerHand Hand, MLInputControllerButton ml_b
 {
 	static const FName empty;
 
+	switch (ml_button)
+	{
+	case MLInputControllerButton_Bumper:
+		if (Hand == EControllerHand::Left)
+		{
+			return FMagicLeapControllerKeyNames::Left_Bumper_Name;
+		}
+		return FMagicLeapControllerKeyNames::Right_Bumper_Name;
+	case MLInputControllerButton_HomeTap:
+		if (Hand == EControllerHand::Left)
+		{
+			return FMagicLeapControllerKeyNames::Left_HomeButton_Name;
+		}
+		return FMagicLeapControllerKeyNames::Right_HomeButton_Name;
+	default:
+		break;
+	}
+	return empty;
+}
+
+const FName& MLToUnrealButtonLegacy(EControllerHand Hand, MLInputControllerButton ml_button)
+{
+	static const FName empty;
+
 	switch (ml_button)
 	{
 	case MLInputControllerButton_Bumper:
@@ -85,8 +109,38 @@ const FName& MLToUnrealButton(FName MotionSource, MLInputControllerButton ml_but
 }
 #endif //WITH_MLSDK
 
+const FName& MLTouchToUnrealTrackpadAxis(EControllerHand Hand, uint32 TouchIndex)
+{
+	static const FName empty;
 
-const FName& MLTouchToUnrealThumbstickAxis(EControllerHand Hand, uint32 TouchIndex)
+	switch (TouchIndex)
+	{
+	case 0:
+		if (Hand == EControllerHand::Left)
+		{
+			return FMagicLeapControllerKeyNames::Left_Trackpad_X_Name;
+		}
+		return FMagicLeapControllerKeyNames::Right_Trackpad_X_Name;
+	case 1:
+		if (Hand == EControllerHand::Left)
+		{
+			return FMagicLeapControllerKeyNames::Left_Trackpad_X_Name;;
+		}
+		return FMagicLeapControllerKeyNames::Right_Trackpad_X_Name;
+	case 2:
+		if (Hand == EControllerHand::Left)
+		{
+			return FMagicLeapControllerKeyNames::Left_Trackpad_Force_Name;
+		}
+		return FMagicLeapControllerKeyNames::Right_Trackpad_Force_Name;
+	default:
+		return empty;
+	}
+
+	return empty;
+}
+
+const FName& MLTouchToUnrealTrackpadAxisLegacy(EControllerHand Hand, uint32 TouchIndex)
 {
 	static const FName empty;
 
@@ -117,7 +171,53 @@ const FName& MLTouchToUnrealThumbstickAxis(EControllerHand Hand, uint32 TouchInd
 	return empty;
 }
 
-const FName& MLTouchToUnrealThumbstickButton(EControllerHand Hand)
+const FName& MLTouchToUnrealTouch1Axis(EControllerHand Hand, uint32 TouchIndex)
+{
+	static const FName empty;
+
+	switch (TouchIndex)
+	{
+	case 0:
+		if (Hand == EControllerHand::Left)
+		{
+			return FMagicLeapControllerKeyNames::Left_Touch1_X_Name;
+		}
+		return FMagicLeapControllerKeyNames::Right_Touch1_X_Name;
+	case 1:
+		if (Hand == EControllerHand::Left)
+		{
+			return FMagicLeapControllerKeyNames::Left_Touch1_Y_Name;
+		}
+		return FMagicLeapControllerKeyNames::Right_Touch1_Y_Name;
+	case 2:
+		if (Hand == EControllerHand::Left)
+		{
+			return FMagicLeapControllerKeyNames::Left_Touch1_Force_Name;
+		}
+		return FMagicLeapControllerKeyNames::Right_Touch1_Force_Name;
+	default:
+		return empty;
+	}
+
+	return empty;
+}
+
+const FName& MLTouchToUnrealTrackpadButton(EControllerHand Hand)
+{
+	static const FName empty;
+
+	switch (Hand)
+	{
+	case EControllerHand::Left:
+		return FMagicLeapControllerKeyNames::Left_Trackpad_Touch_Name;
+	case EControllerHand::Right:
+		return FMagicLeapControllerKeyNames::Right_Trackpad_Touch_Name;
+	}
+
+	return empty;
+}
+
+const FName& MLTouchToUnrealTrackpadButtonLegacy(EControllerHand Hand)
 {
 	static const FName empty;
 
@@ -132,10 +232,40 @@ const FName& MLTouchToUnrealThumbstickButton(EControllerHand Hand)
 	return empty;
 }
 
+const FName& MLTouchToUnrealTouch1Button(EControllerHand Hand)
+{
+	static const FName empty;
+
+	switch (Hand)
+	{
+	case EControllerHand::Left:
+		return FMagicLeapControllerKeyNames::Left_Touch1_Touch_Name;
+	case EControllerHand::Right:
+		return FMagicLeapControllerKeyNames::Right_Touch1_Touch_Name;
+	}
+
+	return empty;
+}
+
 const FName& MLTriggerToUnrealTriggerAxis(EControllerHand Hand)
 {
 	static const FName empty;
 
+	switch (Hand)
+	{
+	case EControllerHand::Left:
+		return FMagicLeapControllerKeyNames::Left_Trigger_Axis_Name;
+	case EControllerHand::Right:
+		return FMagicLeapControllerKeyNames::Right_Trigger_Axis_Name;
+	}
+
+	return empty;
+}
+
+const FName& MLTriggerToUnrealTriggerAxisLegacy(EControllerHand Hand)
+{
+	static const FName empty;
+
 	switch (Hand)
 	{
 	case EControllerHand::Left:
@@ -148,6 +278,19 @@ const FName& MLTriggerToUnrealTriggerAxis(EControllerHand Hand)
 }
 
 const FName& MLTriggerToUnrealTriggerKey(EControllerHand Hand)
+{
+	static const FName empty;
+	switch (Hand)
+	{
+	case EControllerHand::Left:
+		return FMagicLeapControllerKeyNames::Left_Trigger_Name;
+	case EControllerHand::Right:
+		return FMagicLeapControllerKeyNames::Right_Trigger_Name;
+	}
+	return empty;
+}
+
+const FName& MLTriggerToUnrealTriggerKeyLegacy(EControllerHand Hand)
 {
 	static const FName empty;
 	switch (Hand)
diff --git a/Engine/Plugins/Lumin/MagicLeap/Source/MagicLeapController/Public/MagicLeapControllerKeys.h b/Engine/Plugins/Lumin/MagicLeap/Source/MagicLeapController/Public/MagicLeapControllerKeys.h
--- a/Engine/Plugins/Lumin/MagicLeap/Source/MagicLeapController/Public/MagicLeapControllerKeys.h
+++ b/Engine/Plugins/Lumin/MagicLeap/Source/MagicLeapController/Public/MagicLeapControllerKeys.h
@@ -20,9 +20,31 @@ struct FMagicLeapControllerKeyNames
 {
 	static const FName MotionController_Left_Thumbstick_Z_Name;
 	static const FName Left_HomeButton_Name;
+	static const FName Left_Bumper_Name;
+	static const FName Left_Trigger_Name;
+	static const FName Left_Trigger_Axis_Name;
+	static const FName Left_Trackpad_X_Name;
+	static const FName Left_Trackpad_Y_Name;
+	static const FName Left_Trackpad_Force_Name;
+	static const FName Left_Trackpad_Touch_Name;
+	static const FName Left_Touch1_X_Name;
+	static const FName Left_Touch1_Y_Name;
+	static const FName Left_Touch1_Force_Name;
+	static const FName Left_Touch1_Touch_Name;
 
 	static const FName MotionController_Right_Thumbstick_Z_Name;
 	static const FName Right_HomeButton_Name;
+	static const FName Right_Bumper_Name;
+	static const FName Right_Trigger_Name;
+	static const FName Right_Trigger_Axis_Name;
+	static const FName Right_Trackpad_X_Name;
+	static const FName Right_Trackpad_Y_Name;
+	static const FName Right_Trackpad_Force_Name;
+	static const FName Right_Trackpad_Touch_Name;
+	static const FName Right_Touch1_X_Name;
+	static const FName Right_Touch1_Y_Name;
+	static const FName Right_Touch1_Force_Name;
+	static const FName Right_Touch1_Touch_Name;
 
 	static const FName TouchpadGesture_Swipe_Up_Name;
 	static const FName TouchpadGesture_Swipe_Down_Name;
@@ -49,9 +71,31 @@ struct MAGICLEAPCONTROLLER_API FMagicLeapKeys
 {
 	static const FKey MotionController_Left_Thumbstick_Z;
 	static const FKey Left_HomeButton;
+	static const FKey Left_Bumper;
+	static const FKey Left_Trigger;
+	static const FKey Left_Trigger_Axis;
+	static const FKey Left_Trackpad_X;
+	static const FKey Left_Trackpad_Y;
+	static const FKey Left_Trackpad_Force;
+	static const FKey Left_Trackpad_Touch;
+	static const FKey Left_Touch1_X;
+	static const FKey Left_Touch1_Y;
+	static const FKey Left_Touch1_Force;
+	static const FKey Left_Touch1_Touch;
 
 	static const FKey MotionController_Right_Thumbstick_Z;
 	static const FKey Right_HomeButton;
+	static const FKey Right_Bumper;
+	static const FKey Right_Trigger;
+	static const FKey Right_Trigger_Axis;
+	static const FKey Right_Trackpad_X;
+	static const FKey Right_Trackpad_Y;
+	static const FKey Right_Trackpad_Force;
+	static const FKey Right_Trackpad_Touch;
+	static const FKey Right_Touch1_X;
+	static const FKey Right_Touch1_Y;
+	static const FKey Right_Touch1_Force;
+	static const FKey Right_Touch1_Touch;
 
 	static const FKey TouchpadGesture_Swipe_Up;
 	static const FKey TouchpadGesture_Swipe_Down;
diff --git a/Engine/Plugins/Runtime/GoogleVR/GoogleVRController/Source/GoogleVRController/Private/GoogleVRController.cpp b/Engine/Plugins/Runtime/GoogleVR/GoogleVRController/Source/GoogleVRController/Private/GoogleVRController.cpp
--- a/Engine/Plugins/Runtime/GoogleVR/GoogleVRController/Source/GoogleVRController/Private/GoogleVRController.cpp
+++ b/Engine/Plugins/Runtime/GoogleVR/GoogleVRController/Source/GoogleVRController/Private/GoogleVRController.cpp
@@ -226,6 +226,15 @@ FGoogleVRController::FGoogleVRController(const TSharedRef< FGenericApplicationMe
 	Buttons[(int32)EControllerHand::Left][EGoogleVRControllerButton::TouchPadTouch] = GoogleVRControllerKeyNames::Touch0;
 	Buttons[(int32)EControllerHand::Right][EGoogleVRControllerButton::TouchPadTouch] = GoogleVRControllerKeyNames::Touch0;
 
+	Buttons[(int32)EControllerHand::Left][EGoogleVRControllerButton::SelectClick] = EKeys::Daydream_Left_Select_Click.GetFName();
+	Buttons[(int32)EControllerHand::Right][EGoogleVRControllerButton::SelectClick] = EKeys::Daydream_Right_Select_Click.GetFName();
+
+	Buttons[(int32)EControllerHand::Left][EGoogleVRControllerButton::TrackpadClick] = EKeys::Daydream_Left_Trackpad_Click.GetFName();
+	Buttons[(int32)EControllerHand::Right][EGoogleVRControllerButton::TrackpadClick] = EKeys::Daydream_Right_Trackpad_Click.GetFName();
+
+	Buttons[(int32)EControllerHand::Left][EGoogleVRControllerButton::TrackpadTouch] = EKeys::Daydream_Left_Trackpad_Touch.GetFName();
+	Buttons[(int32)EControllerHand::Right][EGoogleVRControllerButton::TrackpadTouch] = EKeys::Daydream_Right_Trackpad_Touch.GetFName();
+
 	// Register callbacks for pause and resume
 	FCoreDelegates::ApplicationWillEnterBackgroundDelegate.AddRaw(this, &FGoogleVRController::ApplicationPauseDelegate);
 	FCoreDelegates::ApplicationHasEnteredForegroundDelegate.AddRaw(this, &FGoogleVRController::ApplicationResumeDelegate);
@@ -456,10 +465,12 @@ void FGoogleVRController::ProcessControllerButtons(int32 ControllerStateIndex)
 			if (CachedControllerState.GetButtonState(gvr::ControllerButton::GVR_CONTROLLER_BUTTON_CLICK))
 			{
 				CurrentButtonStates[EGoogleVRControllerButton::TouchPadPress] = true;
+				CurrentButtonStates[EGoogleVRControllerButton::TrackpadClick] = true;
 			}
 			else if (CachedControllerState.GetButtonUp(gvr::ControllerButton::GVR_CONTROLLER_BUTTON_CLICK))
 			{
 				CurrentButtonStates[EGoogleVRControllerButton::TouchPadPress] = false;
+				CurrentButtonStates[EGoogleVRControllerButton::TrackpadClick] = false;
 			}
 
 			if (CachedControllerState.GetButtonState(gvr::ControllerButton::GVR_CONTROLLER_BUTTON_HOME))
@@ -476,10 +487,12 @@ void FGoogleVRController::ProcessControllerButtons(int32 ControllerStateIndex)
 			if (CachedControllerState.GetButtonState(gvr::ControllerButton::GVR_CONTROLLER_BUTTON_APP))
 			{
 				CurrentButtonStates[EGoogleVRControllerButton::ApplicationMenu] = true;
+				CurrentButtonStates[EGoogleVRControllerButton::SelectClick] = true;
 			}
 			else if (CachedControllerState.GetButtonUp(gvr::ControllerButton::GVR_CONTROLLER_BUTTON_APP))
 			{
 				CurrentButtonStates[EGoogleVRControllerButton::ApplicationMenu] = false;
+				CurrentButtonStates[EGoogleVRControllerButton::SelectClick] = false;
 			}
 
 			if (CachedControllerState.GetButtonState(gvr::ControllerButton::GVR_CONTROLLER_BUTTON_RESERVED0))
@@ -503,6 +516,7 @@ void FGoogleVRController::ProcessControllerButtons(int32 ControllerStateIndex)
 			// Process touches and analog information
 			// OnDown
 			CurrentButtonStates[EGoogleVRControllerButton::TouchPadTouch] = CachedControllerState.IsTouching();
+			CurrentButtonStates[EGoogleVRControllerButton::TrackpadTouch] = CachedControllerState.IsTouching();
 
 			// The controller's touch positions are in [0,1]^2 coordinate space, we want to be in [-1,1]^2, so translate the touch positions.
 			TranslatedLocation = FVector2D((CachedControllerState.GetTouchPos().x * 2) - 1, (CachedControllerState.GetTouchPos().y * 2) - 1);
@@ -541,13 +555,17 @@ void FGoogleVRController::ProcessControllerButtons(int32 ControllerStateIndex)
 
 	if (Hand == EControllerHand::Left)
 	{
-		MessageHandler->OnControllerAnalog( FGamepadKeyNames::MotionController_Left_Thumbstick_X, 0, TranslatedLocation.X);
+		MessageHandler->OnControllerAnalog(FGamepadKeyNames::MotionController_Left_Thumbstick_X, 0, TranslatedLocation.X);
 		MessageHandler->OnControllerAnalog(FGamepadKeyNames::MotionController_Left_Thumbstick_Y, 0, TranslatedLocation.Y);
+		MessageHandler->OnControllerAnalog(EKeys::Daydream_Left_Trackpad_X.GetFName(), 0, TranslatedLocation.X);
+		MessageHandler->OnControllerAnalog(EKeys::Daydream_Left_Trackpad_Y.GetFName(), 0, TranslatedLocation.Y);
 	}
 	if (Hand == EControllerHand::Right)
 	{
 		MessageHandler->OnControllerAnalog(FGamepadKeyNames::MotionController_Right_Thumbstick_X, 0, TranslatedLocation.X);
 		MessageHandler->OnControllerAnalog(FGamepadKeyNames::MotionController_Right_Thumbstick_Y, 0, TranslatedLocation.Y);
+		MessageHandler->OnControllerAnalog(EKeys::Daydream_Right_Trackpad_X.GetFName(), 0, TranslatedLocation.X);
+		MessageHandler->OnControllerAnalog(EKeys::Daydream_Right_Trackpad_Y.GetFName(), 0, TranslatedLocation.Y);
 	}
 
 	// Process buttons
diff --git a/Engine/Plugins/Runtime/GoogleVR/GoogleVRController/Source/GoogleVRController/Private/GoogleVRController.h b/Engine/Plugins/Runtime/GoogleVR/GoogleVRController/Source/GoogleVRController/Private/GoogleVRController.h
--- a/Engine/Plugins/Runtime/GoogleVR/GoogleVRController/Source/GoogleVRController/Private/GoogleVRController.h
+++ b/Engine/Plugins/Runtime/GoogleVR/GoogleVRController/Source/GoogleVRController/Private/GoogleVRController.h
@@ -61,6 +61,7 @@ public:
 	{
 		enum Type
 		{
+			// Deprecated buttons
 			ApplicationMenu,
 			TouchPadLeft,
 			TouchPadUp,
@@ -72,6 +73,11 @@ public:
 			TouchPadPress,
 			TouchPadTouch,
 
+			// New buttons
+			SelectClick,
+			TrackpadClick,
+			TrackpadTouch,
+
 			/** Max number of controller buttons.  Must be < 256 */
 			TotalButtonCount
 		};
diff --git a/Engine/Plugins/Runtime/Oculus/OculusVR/Source/OculusInput/Private/OculusInput.cpp b/Engine/Plugins/Runtime/Oculus/OculusVR/Source/OculusInput/Private/OculusInput.cpp
--- a/Engine/Plugins/Runtime/Oculus/OculusVR/Source/OculusInput/Private/OculusInput.cpp
+++ b/Engine/Plugins/Runtime/Oculus/OculusVR/Source/OculusInput/Private/OculusInput.cpp
@@ -79,6 +79,7 @@ float FOculusInput::TriggerThreshold = 0.8f;
 
 /** Are Remote keys mapped to gamepad or not. */
 bool FOculusInput::bRemoteKeysMappedToGamepad = true;
+bool FOculusInput::bGoKeysMappedToTouch = false;
 
 float FOculusInput::InitialButtonRepeatDelay = 0.2f;
 float FOculusInput::ButtonRepeatDelay = 0.1f;
@@ -124,36 +125,38 @@ void FOculusInput::PreInit()
 	LoadConfig();
 
 	// Register the FKeys
-	EKeys::AddKey(FKeyDetails(FOculusKey::OculusTouch_Left_Thumbstick, LOCTEXT("OculusTouch_Left_Thumbstick", "Oculus Touch (L) Thumbstick CapTouch"), FKeyDetails::GamepadKey | FKeyDetails::FloatAxis));
-	EKeys::AddKey(FKeyDetails(FOculusKey::OculusTouch_Left_FaceButton1, LOCTEXT("OculusTouch_Left_FaceButton1", "Oculus Touch (L) X Button CapTouch"), FKeyDetails::GamepadKey | FKeyDetails::FloatAxis));
-	EKeys::AddKey(FKeyDetails(FOculusKey::OculusTouch_Left_Trigger, LOCTEXT("OculusTouch_Left_Trigger", "Oculus Touch (L) Trigger CapTouch"), FKeyDetails::GamepadKey | FKeyDetails::FloatAxis));
-	EKeys::AddKey(FKeyDetails(FOculusKey::OculusTouch_Left_FaceButton2, LOCTEXT("OculusTouch_Left_FaceButton2", "Oculus Touch (L) Y Button CapTouch"), FKeyDetails::GamepadKey | FKeyDetails::FloatAxis));
-	EKeys::AddKey(FKeyDetails(FOculusKey::OculusTouch_Left_IndexPointing, LOCTEXT("OculusTouch_Left_IndexPointing", "Oculus Touch (L) Pointing CapTouch"), FKeyDetails::GamepadKey | FKeyDetails::FloatAxis));
-	EKeys::AddKey(FKeyDetails(FOculusKey::OculusTouch_Left_ThumbUp, LOCTEXT("OculusTouch_Left_ThumbUp", "Oculus Touch (L) Thumb Up CapTouch"), FKeyDetails::GamepadKey | FKeyDetails::FloatAxis));
-
-	EKeys::AddKey(FKeyDetails(FOculusKey::OculusTouch_Right_Thumbstick, LOCTEXT("OculusTouch_Right_Thumbstick", "Oculus Touch (R) Thumbstick CapTouch"), FKeyDetails::GamepadKey | FKeyDetails::FloatAxis));
-	EKeys::AddKey(FKeyDetails(FOculusKey::OculusTouch_Right_FaceButton1, LOCTEXT("OculusTouch_Right_FaceButton1", "Oculus Touch (R) A Button CapTouch"), FKeyDetails::GamepadKey | FKeyDetails::FloatAxis));
-	EKeys::AddKey(FKeyDetails(FOculusKey::OculusTouch_Right_Trigger, LOCTEXT("OculusTouch_Right_Trigger", "Oculus Touch (R) Trigger CapTouch"), FKeyDetails::GamepadKey | FKeyDetails::FloatAxis));
-	EKeys::AddKey(FKeyDetails(FOculusKey::OculusTouch_Right_FaceButton2, LOCTEXT("OculusTouch_Right_FaceButton2", "Oculus Touch (R) B Button CapTouch"), FKeyDetails::GamepadKey | FKeyDetails::FloatAxis));
-	EKeys::AddKey(FKeyDetails(FOculusKey::OculusTouch_Right_IndexPointing, LOCTEXT("OculusTouch_Right_IndexPointing", "Oculus Touch (R) Pointing CapTouch"), FKeyDetails::GamepadKey | FKeyDetails::FloatAxis));
-	EKeys::AddKey(FKeyDetails(FOculusKey::OculusTouch_Right_ThumbUp, LOCTEXT("OculusTouch_Right_ThumbUp", "Oculus Touch (R) Thumb Up CapTouch"), FKeyDetails::GamepadKey | FKeyDetails::FloatAxis));
-
-	EKeys::AddKey(FKeyDetails(FOculusKey::OculusRemote_DPad_Up, LOCTEXT("OculusRemote_DPad_Up", "Oculus Remote D-pad Up"), FKeyDetails::GamepadKey));
-	EKeys::AddKey(FKeyDetails(FOculusKey::OculusRemote_DPad_Down, LOCTEXT("OculusRemote_DPad_Down", "Oculus Remote D-pad Down"), FKeyDetails::GamepadKey));
-	EKeys::AddKey(FKeyDetails(FOculusKey::OculusRemote_DPad_Left, LOCTEXT("OculusRemote_DPad_Left", "Oculus Remote D-pad Left"), FKeyDetails::GamepadKey));
-	EKeys::AddKey(FKeyDetails(FOculusKey::OculusRemote_DPad_Right, LOCTEXT("OculusRemote_DPad_Right", "Oculus Remote D-pad Right"), FKeyDetails::GamepadKey));
-
-	EKeys::AddKey(FKeyDetails(FOculusKey::OculusRemote_Enter, LOCTEXT("OculusRemote_Enter", "Oculus Remote Enter"), FKeyDetails::GamepadKey));
-	EKeys::AddKey(FKeyDetails(FOculusKey::OculusRemote_Back, LOCTEXT("OculusRemote_Back", "Oculus Remote Back"), FKeyDetails::GamepadKey));
-
-	EKeys::AddKey(FKeyDetails(FOculusKey::OculusRemote_VolumeUp, LOCTEXT("OculusRemote_VolumeUp", "Oculus Remote Volume Up"), FKeyDetails::GamepadKey));
-	EKeys::AddKey(FKeyDetails(FOculusKey::OculusRemote_VolumeDown, LOCTEXT("OculusRemote_VolumeDown", "Oculus Remote Volume Down"), FKeyDetails::GamepadKey));
-	EKeys::AddKey(FKeyDetails(FOculusKey::OculusRemote_Home, LOCTEXT("OculusRemote_Home", "Oculus Remote Home"), FKeyDetails::GamepadKey));
-
-	EKeys::AddKey(FKeyDetails(FOculusKey::OculusTouchpad_Touchpad, LOCTEXT("OculusTouchpad_Touchpad", "Oculus Touchpad Button"), FKeyDetails::GamepadKey));
-	EKeys::AddKey(FKeyDetails(FOculusKey::OculusTouchpad_Touchpad_X, LOCTEXT("OculusTouchpad_Touchpad_X", "Oculus Touchpad X-Axis"), FKeyDetails::GamepadKey | FKeyDetails::FloatAxis));
-	EKeys::AddKey(FKeyDetails(FOculusKey::OculusTouchpad_Touchpad_Y, LOCTEXT("OculusTouchpad_Touchpad_Y", "Oculus Touchpad Y-Axis"), FKeyDetails::GamepadKey | FKeyDetails::FloatAxis));
-	EKeys::AddKey(FKeyDetails(FOculusKey::OculusTouchpad_Back, LOCTEXT("OculusTouchpad_Back", "Oculus Touchpad Back"), FKeyDetails::GamepadKey));
+	EKeys::AddKey(FKeyDetails(FOculusKey::OculusTouch_Left_Thumbstick, LOCTEXT("OculusTouch_Left_Thumbstick", "Oculus Touch (L) Thumbstick CapTouch"), FKeyDetails::GamepadKey | FKeyDetails::FloatAxis | FKeyDetails::NotBlueprintBindableKey | FKeyDetails::NotActionBindableKey));
+	EKeys::AddKey(FKeyDetails(FOculusKey::OculusTouch_Left_FaceButton1, LOCTEXT("OculusTouch_Left_FaceButton1", "Oculus Touch (L) X Button CapTouch"), FKeyDetails::GamepadKey | FKeyDetails::FloatAxis | FKeyDetails::NotBlueprintBindableKey | FKeyDetails::NotActionBindableKey));
+	EKeys::AddKey(FKeyDetails(FOculusKey::OculusTouch_Left_Trigger, LOCTEXT("OculusTouch_Left_Trigger", "Oculus Touch (L) Trigger CapTouch"), FKeyDetails::GamepadKey | FKeyDetails::FloatAxis | FKeyDetails::NotBlueprintBindableKey | FKeyDetails::NotActionBindableKey));
+	EKeys::AddKey(FKeyDetails(FOculusKey::OculusTouch_Left_FaceButton2, LOCTEXT("OculusTouch_Left_FaceButton2", "Oculus Touch (L) Y Button CapTouch"), FKeyDetails::GamepadKey | FKeyDetails::FloatAxis | FKeyDetails::NotBlueprintBindableKey | FKeyDetails::NotActionBindableKey));
+	EKeys::AddKey(FKeyDetails(FOculusKey::OculusTouch_Left_IndexPointing, LOCTEXT("OculusTouch_Left_IndexPointing", "Oculus Touch (L) Pointing CapTouch"), FKeyDetails::GamepadKey | FKeyDetails::FloatAxis | FKeyDetails::NotBlueprintBindableKey, "OculusTouch"));
+	EKeys::AddKey(FKeyDetails(FOculusKey::OculusTouch_Left_ThumbUp, LOCTEXT("OculusTouch_Left_ThumbUp", "Oculus Touch (L) Thumb Up CapTouch"), FKeyDetails::GamepadKey | FKeyDetails::FloatAxis | FKeyDetails::NotBlueprintBindableKey, "OculusTouch"));
+
+	EKeys::AddKey(FKeyDetails(FOculusKey::OculusTouch_Right_Thumbstick, LOCTEXT("OculusTouch_Right_Thumbstick", "Oculus Touch (R) Thumbstick CapTouch"), FKeyDetails::GamepadKey | FKeyDetails::FloatAxis | FKeyDetails::NotBlueprintBindableKey));
+	EKeys::AddKey(FKeyDetails(FOculusKey::OculusTouch_Right_FaceButton1, LOCTEXT("OculusTouch_Right_FaceButton1", "Oculus Touch (R) A Button CapTouch"), FKeyDetails::GamepadKey | FKeyDetails::FloatAxis | FKeyDetails::NotBlueprintBindableKey));
+	EKeys::AddKey(FKeyDetails(FOculusKey::OculusTouch_Right_Trigger, LOCTEXT("OculusTouch_Right_Trigger", "Oculus Touch (R) Trigger CapTouch"), FKeyDetails::GamepadKey | FKeyDetails::FloatAxis | FKeyDetails::NotBlueprintBindableKey));
+	EKeys::AddKey(FKeyDetails(FOculusKey::OculusTouch_Right_FaceButton2, LOCTEXT("OculusTouch_Right_FaceButton2", "Oculus Touch (R) B Button CapTouch"), FKeyDetails::GamepadKey | FKeyDetails::FloatAxis | FKeyDetails::NotBlueprintBindableKey));
+	EKeys::AddKey(FKeyDetails(FOculusKey::OculusTouch_Right_IndexPointing, LOCTEXT("OculusTouch_Right_IndexPointing", "Oculus Touch (R) Pointing CapTouch"), FKeyDetails::GamepadKey | FKeyDetails::FloatAxis | FKeyDetails::NotBlueprintBindableKey, "OculusTouch"));
+	EKeys::AddKey(FKeyDetails(FOculusKey::OculusTouch_Right_ThumbUp, LOCTEXT("OculusTouch_Right_ThumbUp", "Oculus Touch (R) Thumb Up CapTouch"), FKeyDetails::GamepadKey | FKeyDetails::FloatAxis | FKeyDetails::NotBlueprintBindableKey, "OculusTouch"));
+
+	EKeys::AddMenuCategoryDisplayInfo("OculusRemote", LOCTEXT("OculusRemoteSubCategory", "Oculus Remote"), TEXT("GraphEditor.PadEvent_16x"));
+
+	EKeys::AddKey(FKeyDetails(FOculusKey::OculusRemote_DPad_Up, LOCTEXT("OculusRemote_DPad_Up", "Oculus Remote D-pad Up"), FKeyDetails::GamepadKey, "OculusRemote"));
+	EKeys::AddKey(FKeyDetails(FOculusKey::OculusRemote_DPad_Down, LOCTEXT("OculusRemote_DPad_Down", "Oculus Remote D-pad Down"), FKeyDetails::GamepadKey, "OculusRemote"));
+	EKeys::AddKey(FKeyDetails(FOculusKey::OculusRemote_DPad_Left, LOCTEXT("OculusRemote_DPad_Left", "Oculus Remote D-pad Left"), FKeyDetails::GamepadKey, "OculusRemote"));
+	EKeys::AddKey(FKeyDetails(FOculusKey::OculusRemote_DPad_Right, LOCTEXT("OculusRemote_DPad_Right", "Oculus Remote D-pad Right"), FKeyDetails::GamepadKey, "OculusRemote"));
+
+	EKeys::AddKey(FKeyDetails(FOculusKey::OculusRemote_Enter, LOCTEXT("OculusRemote_Enter", "Oculus Remote Enter"), FKeyDetails::GamepadKey, "OculusRemote"));
+	EKeys::AddKey(FKeyDetails(FOculusKey::OculusRemote_Back, LOCTEXT("OculusRemote_Back", "Oculus Remote Back"), FKeyDetails::GamepadKey, "OculusRemote"));
+
+	EKeys::AddKey(FKeyDetails(FOculusKey::OculusRemote_VolumeUp, LOCTEXT("OculusRemote_VolumeUp", "Oculus Remote Volume Up"), FKeyDetails::GamepadKey, "OculusRemote"));
+	EKeys::AddKey(FKeyDetails(FOculusKey::OculusRemote_VolumeDown, LOCTEXT("OculusRemote_VolumeDown", "Oculus Remote Volume Down"), FKeyDetails::GamepadKey, "OculusRemote"));
+	EKeys::AddKey(FKeyDetails(FOculusKey::OculusRemote_Home, LOCTEXT("OculusRemote_Home", "Oculus Remote Home"), FKeyDetails::GamepadKey, "OculusRemote"));
+
+	EKeys::AddKey(FKeyDetails(FOculusKey::OculusTouchpad_Touchpad, LOCTEXT("OculusTouchpad_Touchpad", "Oculus Touchpad Button"), FKeyDetails::GamepadKey | FKeyDetails::NotBlueprintBindableKey | FKeyDetails::NotActionBindableKey));
+	EKeys::AddKey(FKeyDetails(FOculusKey::OculusTouchpad_Touchpad_X, LOCTEXT("OculusTouchpad_Touchpad_X", "Oculus Touchpad X-Axis"), FKeyDetails::GamepadKey | FKeyDetails::FloatAxis | FKeyDetails::NotBlueprintBindableKey | FKeyDetails::NotActionBindableKey));
+	EKeys::AddKey(FKeyDetails(FOculusKey::OculusTouchpad_Touchpad_Y, LOCTEXT("OculusTouchpad_Touchpad_Y", "Oculus Touchpad Y-Axis"), FKeyDetails::GamepadKey | FKeyDetails::FloatAxis | FKeyDetails::NotBlueprintBindableKey | FKeyDetails::NotActionBindableKey));
+	EKeys::AddKey(FKeyDetails(FOculusKey::OculusTouchpad_Back, LOCTEXT("OculusTouchpad_Back", "Oculus Touchpad Back"), FKeyDetails::GamepadKey | FKeyDetails::NotBlueprintBindableKey | FKeyDetails::NotActionBindableKey));
 
 	UE_LOG(LogOcInput, Log, TEXT("OculusInput pre-init called"));
 }
@@ -162,17 +165,22 @@ void FOculusInput::LoadConfig()
 {
 	const TCHAR* OculusTouchSettings = TEXT("OculusTouch.Settings");
 	float ConfigThreshold = TriggerThreshold;
-
 	if (GConfig->GetFloat(OculusTouchSettings, TEXT("TriggerThreshold"), ConfigThreshold, GEngineIni))
 	{
 		TriggerThreshold = ConfigThreshold;
 	}
 
+	bool bConfigGoKeysMappedToTouch;
+	if (GConfig->GetBool(OculusTouchSettings, TEXT("bGoKeysMappedToTouch"), bConfigGoKeysMappedToTouch, GEngineIni))
+	{
+		bGoKeysMappedToTouch = bConfigGoKeysMappedToTouch;
+	}
+
 	const TCHAR* OculusRemoteSettings = TEXT("OculusRemote.Settings");
-	bool b;
-	if (GConfig->GetBool(OculusRemoteSettings, TEXT("bRemoteKeysMappedToGamepad"), b, GEngineIni))
+	bool bConfigRemoteKeysMappedToGamepad;
+	if (GConfig->GetBool(OculusRemoteSettings, TEXT("bRemoteKeysMappedToGamepad"), bConfigRemoteKeysMappedToGamepad, GEngineIni))
 	{
-		bRemoteKeysMappedToGamepad = b;
+		bRemoteKeysMappedToGamepad = bConfigRemoteKeysMappedToGamepad;
 	}
 
 	GConfig->GetFloat(TEXT("/Script/Engine.InputSettings"), TEXT("InitialButtonRepeatDelay"), InitialButtonRepeatDelay, GInputIni);
@@ -398,20 +406,28 @@ void FOculusInput::SendControllerEvents()
 							{
 								State.TriggerAxis = OvrTriggerAxis;
 								MessageHandler->OnControllerAnalog(bIsLeft ? FGamepadKeyNames::MotionController_Left_TriggerAxis : FGamepadKeyNames::MotionController_Right_TriggerAxis, ControllerPair.UnrealControllerIndex, State.TriggerAxis);
+								MessageHandler->OnControllerAnalog(bIsLeft ? EKeys::OculusTouch_Left_Trigger_Axis.GetFName() : EKeys::OculusTouch_Right_Trigger_Axis.GetFName(), ControllerPair.UnrealControllerIndex, State.TriggerAxis);
 							}
 
 							if (OvrGripAxis != State.GripAxis)
 							{
 								State.GripAxis = OvrGripAxis;
 								MessageHandler->OnControllerAnalog(bIsLeft ? FGamepadKeyNames::MotionController_Left_Grip1Axis : FGamepadKeyNames::MotionController_Right_Grip1Axis, ControllerPair.UnrealControllerIndex, State.GripAxis);
+								MessageHandler->OnControllerAnalog(bIsLeft ? EKeys::OculusTouch_Left_Grip_Axis.GetFName() : EKeys::OculusTouch_Right_Grip_Axis.GetFName(), ControllerPair.UnrealControllerIndex, State.GripAxis);
 							}
 
-							ovrpVector2f ThumbstickValue = bIsMobileController ? OvrpControllerState.Touchpad[HandIndex] : OvrpControllerState.Thumbstick[HandIndex];
+							ovrpVector2f ThumbstickValue = OvrpControllerState.Thumbstick[HandIndex];
+							ovrpVector2f TouchpadValue = OvrpControllerState.Touchpad[HandIndex];
 
 							if (ThumbstickValue.x != State.ThumbstickAxes.X)
 							{
 								State.ThumbstickAxes.X = ThumbstickValue.x;
 								MessageHandler->OnControllerAnalog(bIsLeft ? FGamepadKeyNames::MotionController_Left_Thumbstick_X : FGamepadKeyNames::MotionController_Right_Thumbstick_X, ControllerPair.UnrealControllerIndex, State.ThumbstickAxes.X);
+								MessageHandler->OnControllerAnalog(bIsLeft ? EKeys::OculusTouch_Left_Thumbstick_X.GetFName() : EKeys::OculusTouch_Right_Thumbstick_X.GetFName(), ControllerPair.UnrealControllerIndex, State.ThumbstickAxes.X);
+								if (bGoKeysMappedToTouch)
+								{
+									MessageHandler->OnControllerAnalog(bIsLeft ? EKeys::OculusGo_Left_Trackpad_X.GetFName() : EKeys::OculusGo_Right_Trackpad_X.GetFName(), ControllerPair.UnrealControllerIndex, State.ThumbstickAxes.Y);
+								}
 							}
 
 							if (ThumbstickValue.y != State.ThumbstickAxes.Y)
@@ -419,6 +435,30 @@ void FOculusInput::SendControllerEvents()
 								State.ThumbstickAxes.Y = ThumbstickValue.y;
 								// we need to negate Y value to match XBox controllers
 								MessageHandler->OnControllerAnalog(bIsLeft ? FGamepadKeyNames::MotionController_Left_Thumbstick_Y : FGamepadKeyNames::MotionController_Right_Thumbstick_Y, ControllerPair.UnrealControllerIndex, -State.ThumbstickAxes.Y);
+
+								// new keys no longer need negation
+								MessageHandler->OnControllerAnalog(bIsLeft ? EKeys::OculusTouch_Left_Thumbstick_Y.GetFName() : EKeys::OculusTouch_Right_Thumbstick_Y.GetFName(), ControllerPair.UnrealControllerIndex, State.ThumbstickAxes.Y);
+								if (bGoKeysMappedToTouch)
+								{
+									MessageHandler->OnControllerAnalog(bIsLeft ? EKeys::OculusGo_Left_Trackpad_Y.GetFName() : EKeys::OculusGo_Right_Trackpad_Y.GetFName(), ControllerPair.UnrealControllerIndex, State.ThumbstickAxes.Y);
+								}
+							}
+
+							if (TouchpadValue.x != State.TouchpadAxes.X)
+							{
+								State.TouchpadAxes.X = TouchpadValue.x;
+								MessageHandler->OnControllerAnalog(bIsLeft ? FGamepadKeyNames::MotionController_Left_Thumbstick_X : FGamepadKeyNames::MotionController_Right_Thumbstick_X, ControllerPair.UnrealControllerIndex, State.ThumbstickAxes.X);
+								MessageHandler->OnControllerAnalog(bIsLeft ? EKeys::OculusGo_Left_Trackpad_X.GetFName() : EKeys::OculusGo_Right_Trackpad_X.GetFName(), ControllerPair.UnrealControllerIndex, State.ThumbstickAxes.X);
+							}
+
+							if (TouchpadValue.y != State.TouchpadAxes.Y)
+							{
+								State.TouchpadAxes.Y = TouchpadValue.y;
+								// we need to negate Y value to match XBox controllers
+								MessageHandler->OnControllerAnalog(bIsLeft ? FGamepadKeyNames::MotionController_Left_Thumbstick_Y : FGamepadKeyNames::MotionController_Right_Thumbstick_Y, ControllerPair.UnrealControllerIndex, -State.ThumbstickAxes.Y);
+
+								// new keys no longer need negation
+								MessageHandler->OnControllerAnalog(bIsLeft ? EKeys::OculusGo_Left_Trackpad_Y.GetFName() : EKeys::OculusGo_Right_Trackpad_Y.GetFName(), ControllerPair.UnrealControllerIndex, State.ThumbstickAxes.Y);
 							}
 
 							for (int32 ButtonIndex = 0; ButtonIndex < (int32)EOculusTouchControllerButton::TotalButtonCount; ++ButtonIndex)
@@ -439,10 +479,7 @@ void FOculusInput::SendControllerEvents()
 									break;
 
 								case EOculusTouchControllerButton::XA:
-									bButtonPressed =
-										bIsMobileController ?
-										(OvrpControllerState.Buttons & ovrpButton_Back) != 0 : 
-										(bIsLeft ? (OvrpControllerState.Buttons & ovrpButton_X) != 0 : (OvrpControllerState.Buttons & ovrpButton_A) != 0);
+									bButtonPressed = bIsLeft ? (OvrpControllerState.Buttons & ovrpButton_X) != 0 : (OvrpControllerState.Buttons & ovrpButton_A) != 0;
 									break;
 
 								case EOculusTouchControllerButton::YB:
@@ -450,10 +487,7 @@ void FOculusInput::SendControllerEvents()
 									break;
 
 								case EOculusTouchControllerButton::Thumbstick:
-									bButtonPressed = 
-										bIsMobileController ?
-											(bIsLeft ? (OvrpControllerState.Buttons & ovrpButton_LTouchpad) != 0 : (OvrpControllerState.Buttons & ovrpButton_RTouchpad) != 0) :
-											(bIsLeft ? (OvrpControllerState.Buttons & ovrpButton_LThumb) != 0 : (OvrpControllerState.Buttons & ovrpButton_RThumb) != 0);
+									bButtonPressed = bIsLeft ? (OvrpControllerState.Buttons & ovrpButton_LThumb) != 0 : (OvrpControllerState.Buttons & ovrpButton_RThumb) != 0;
 									break;
 									
 								case EOculusTouchControllerButton::Thumbstick_Up:
@@ -495,6 +529,37 @@ void FOculusInput::SendControllerEvents()
 								case EOculusTouchControllerButton::Menu:
 									bButtonPressed = bIsLeft && (OvrpControllerState.Buttons & ovrpButton_Start);
 									break;
+
+								case EOculusTouchControllerButton::Thumbstick_Touch:
+									bButtonPressed = bIsLeft ? (OvrpControllerState.Touches & ovrpTouch_LThumb) != 0 : (OvrpControllerState.Touches & ovrpTouch_RThumb) != 0;
+									break;
+
+								case EOculusTouchControllerButton::Trigger_Touch:
+									bButtonPressed = bIsLeft ? (OvrpControllerState.Touches & ovrpTouch_LIndexTrigger) != 0 : (OvrpControllerState.Touches & ovrpTouch_RIndexTrigger) != 0;
+									break;
+
+								case EOculusTouchControllerButton::XA_Touch:
+									bButtonPressed = bIsLeft ? (OvrpControllerState.Touches & ovrpTouch_X) != 0 : (OvrpControllerState.Touches & ovrpTouch_A) != 0;
+									break;
+
+								case EOculusTouchControllerButton::YB_Touch:
+									bButtonPressed = bIsLeft ? (OvrpControllerState.Touches & ovrpTouch_Y) != 0 : (OvrpControllerState.Touches & ovrpTouch_B) != 0;
+									break;
+
+								case EOculusTouchControllerButton::Back:
+									bButtonPressed = (OvrpControllerState.Buttons & ovrpButton_Back) != 0 || (bGoKeysMappedToTouch &&
+										(bIsLeft ? (OvrpControllerState.Buttons & ovrpButton_X) != 0 : (OvrpControllerState.Buttons & ovrpButton_A) != 0));
+									break;
+
+								case EOculusTouchControllerButton::Touchpad:
+									bButtonPressed = (bIsLeft ? (OvrpControllerState.Buttons & ovrpButton_LTouchpad) != 0 : (OvrpControllerState.Buttons & ovrpButton_RTouchpad) != 0) ||
+										(bGoKeysMappedToTouch && (bIsLeft ? (OvrpControllerState.Buttons & ovrpButton_LThumb) != 0 : (OvrpControllerState.Buttons & ovrpButton_RThumb) != 0));
+									break;
+
+								case EOculusTouchControllerButton::Touchpad_Touch:
+									bButtonPressed = (bIsLeft ? (OvrpControllerState.Buttons & ovrpTouch_LTouchpad) != 0 : (OvrpControllerState.Buttons & ovrpTouch_RTouchpad) != 0) ||
+										(bGoKeysMappedToTouch && (bIsLeft ? (OvrpControllerState.Touches & ovrpTouch_LThumb) != 0 : (OvrpControllerState.Touches & ovrpTouch_RThumb) != 0));
+									break;
 								
 								default:
 									check(0);
diff --git a/Engine/Plugins/Runtime/Oculus/OculusVR/Source/OculusInput/Private/OculusInput.h b/Engine/Plugins/Runtime/Oculus/OculusVR/Source/OculusInput/Private/OculusInput.h
--- a/Engine/Plugins/Runtime/Oculus/OculusVR/Source/OculusInput/Private/OculusInput.h
+++ b/Engine/Plugins/Runtime/Oculus/OculusVR/Source/OculusInput/Private/OculusInput.h
@@ -95,6 +95,9 @@ private:
 	/** Are Remote keys mapped to gamepad or not. */
 	static bool bRemoteKeysMappedToGamepad;
 
+	/** Are Go keys mapped to Touch or not. */
+	static bool bGoKeysMappedToTouch;
+
 	/** Repeat key delays, loaded from config */
 	static float InitialButtonRepeatDelay;
 	static float ButtonRepeatDelay;
diff --git a/Engine/Plugins/Runtime/Oculus/OculusVR/Source/OculusInput/Private/OculusInputState.h b/Engine/Plugins/Runtime/Oculus/OculusVR/Source/OculusInput/Private/OculusInputState.h
--- a/Engine/Plugins/Runtime/Oculus/OculusVR/Source/OculusInput/Private/OculusInputState.h
+++ b/Engine/Plugins/Runtime/Oculus/OculusVR/Source/OculusInput/Private/OculusInputState.h
@@ -33,6 +33,16 @@ enum class EOculusTouchControllerButton
 
 	Menu,
 
+	Thumbstick_Touch,
+	Trigger_Touch,
+	XA_Touch,
+	YB_Touch,
+
+	/** Oculus Go buttons */
+	Back,
+	Touchpad,
+	Touchpad_Touch,
+
 	/** Total number of controller buttons */
 	TotalButtonCount
 };
@@ -238,6 +248,9 @@ struct FOculusTouchControllerState
 	/** Thumbstick */
 	FVector2D ThumbstickAxes;
 
+	/** Thumbstick */
+	FVector2D TouchpadAxes;
+
 	/** Button states */
 	FOculusButtonState Buttons[ (int32)EOculusTouchControllerButton::TotalButtonCount ];
 
@@ -286,11 +299,17 @@ struct FOculusTouchControllerState
 			Button.NextRepeatTime = 0.0;
 		}
 
-		Buttons[ (int32)EOculusTouchControllerButton::Trigger ].Key = ( Hand == EControllerHand::Left ) ? FGamepadKeyNames::MotionController_Left_Trigger : FGamepadKeyNames::MotionController_Right_Trigger;
-		Buttons[ (int32)EOculusTouchControllerButton::Grip ].Key = ( Hand == EControllerHand::Left ) ? FGamepadKeyNames::MotionController_Left_Grip1 : FGamepadKeyNames::MotionController_Right_Grip1;
-		Buttons[ (int32)EOculusTouchControllerButton::Thumbstick ].Key = ( Hand == EControllerHand::Left ) ? FGamepadKeyNames::MotionController_Left_Thumbstick : FGamepadKeyNames::MotionController_Right_Thumbstick;
-		Buttons[ (int32)EOculusTouchControllerButton::XA ].Key = (Hand == EControllerHand::Left) ? FGamepadKeyNames::MotionController_Left_FaceButton1 : FGamepadKeyNames::MotionController_Right_FaceButton1;
-		Buttons[ (int32)EOculusTouchControllerButton::YB ].Key = (Hand == EControllerHand::Left) ? FGamepadKeyNames::MotionController_Left_FaceButton2 : FGamepadKeyNames::MotionController_Right_FaceButton2;
+		Buttons[(int32)EOculusTouchControllerButton::Trigger].Key = (Hand == EControllerHand::Left) ? EKeys::OculusTouch_Left_Trigger_Click.GetFName() : EKeys::OculusTouch_Right_Trigger_Click.GetFName();
+		Buttons[(int32)EOculusTouchControllerButton::Grip].Key = (Hand == EControllerHand::Left) ? EKeys::OculusTouch_Left_Grip_Click.GetFName() : EKeys::OculusTouch_Right_Grip_Click.GetFName();
+		Buttons[(int32)EOculusTouchControllerButton::Thumbstick].Key = (Hand == EControllerHand::Left) ? EKeys::OculusTouch_Left_Thumbstick_Click.GetFName() : EKeys::OculusTouch_Right_Thumbstick_Click.GetFName();
+		Buttons[(int32)EOculusTouchControllerButton::XA].Key = (Hand == EControllerHand::Left) ? EKeys::OculusTouch_Left_X_Click.GetFName() : EKeys::OculusTouch_Right_A_Click.GetFName();
+		Buttons[(int32)EOculusTouchControllerButton::YB].Key = (Hand == EControllerHand::Left) ? EKeys::OculusTouch_Left_Y_Click.GetFName() : EKeys::OculusTouch_Right_B_Click.GetFName();
+
+		Buttons[ (int32)EOculusTouchControllerButton::Trigger ].EmulatedKey = ( Hand == EControllerHand::Left ) ? FGamepadKeyNames::MotionController_Left_Trigger : FGamepadKeyNames::MotionController_Right_Trigger;
+		Buttons[ (int32)EOculusTouchControllerButton::Grip ].EmulatedKey = ( Hand == EControllerHand::Left ) ? FGamepadKeyNames::MotionController_Left_Grip1 : FGamepadKeyNames::MotionController_Right_Grip1;
+		Buttons[ (int32)EOculusTouchControllerButton::Thumbstick ].EmulatedKey = ( Hand == EControllerHand::Left ) ? FGamepadKeyNames::MotionController_Left_Thumbstick : FGamepadKeyNames::MotionController_Right_Thumbstick;
+		Buttons[ (int32)EOculusTouchControllerButton::XA ].EmulatedKey = (Hand == EControllerHand::Left) ? FGamepadKeyNames::MotionController_Left_FaceButton1 : FGamepadKeyNames::MotionController_Right_FaceButton1;
+		Buttons[ (int32)EOculusTouchControllerButton::YB ].EmulatedKey = (Hand == EControllerHand::Left) ? FGamepadKeyNames::MotionController_Left_FaceButton2 : FGamepadKeyNames::MotionController_Right_FaceButton2;
 		Buttons[ (int32)EOculusTouchControllerButton::Thumbstick_Up].Key = (Hand == EControllerHand::Left) ? FGamepadKeyNames::MotionController_Left_Thumbstick_Up : FGamepadKeyNames::MotionController_Right_Thumbstick_Up;
 		Buttons[ (int32)EOculusTouchControllerButton::Thumbstick_Down].Key = (Hand == EControllerHand::Left) ? FGamepadKeyNames::MotionController_Left_Thumbstick_Down : FGamepadKeyNames::MotionController_Right_Thumbstick_Down;
 		Buttons[ (int32)EOculusTouchControllerButton::Thumbstick_Left].Key = (Hand == EControllerHand::Left) ? FGamepadKeyNames::MotionController_Left_Thumbstick_Left : FGamepadKeyNames::MotionController_Right_Thumbstick_Left;
@@ -298,6 +317,18 @@ struct FOculusTouchControllerState
 
 		Buttons[ (int32)EOculusTouchControllerButton::Menu ].Key = (Hand == EControllerHand::Left) ? FGamepadKeyNames::SpecialLeft : FGamepadKeyNames::SpecialRight;
 
+		Buttons[ (int32)EOculusTouchControllerButton::Thumbstick_Touch ].Key = (Hand == EControllerHand::Left) ? EKeys::OculusTouch_Left_Thumbstick_Touch.GetFName() : EKeys::OculusTouch_Right_Thumbstick_Touch.GetFName();
+		Buttons[ (int32)EOculusTouchControllerButton::Trigger_Touch ].Key = (Hand == EControllerHand::Left) ? EKeys::OculusTouch_Left_Trigger_Touch.GetFName() : EKeys::OculusTouch_Right_Trigger_Touch.GetFName();
+		Buttons[ (int32)EOculusTouchControllerButton::XA_Touch ].Key = (Hand == EControllerHand::Left) ? EKeys::OculusTouch_Left_X_Touch.GetFName() : EKeys::OculusTouch_Right_A_Touch.GetFName();
+		Buttons[ (int32)EOculusTouchControllerButton::YB_Touch ].Key = (Hand == EControllerHand::Left) ? EKeys::OculusTouch_Left_Y_Touch.GetFName() : EKeys::OculusTouch_Right_B_Touch.GetFName();
+
+		Buttons[ (int32)EOculusTouchControllerButton::Back ].Key = (Hand == EControllerHand::Left) ? EKeys::OculusGo_Left_Back_Click.GetFName() : EKeys::OculusGo_Right_Back_Click.GetFName();
+		Buttons[ (int32)EOculusTouchControllerButton::Touchpad ].Key = (Hand == EControllerHand::Left) ? EKeys::OculusGo_Left_Trackpad_Click.GetFName() : EKeys::OculusGo_Right_Trackpad_Click.GetFName();
+		Buttons[ (int32)EOculusTouchControllerButton::Touchpad_Touch ].Key = (Hand == EControllerHand::Left) ? EKeys::OculusGo_Left_Trackpad_Touch.GetFName() : EKeys::OculusGo_Right_Trackpad_Touch.GetFName();
+
+		Buttons[ (int32)EOculusTouchControllerButton::Back ].EmulatedKey = (Hand == EControllerHand::Left) ? EKeys::MotionController_Left_FaceButton1.GetFName() : EKeys::MotionController_Right_FaceButton1.GetFName();
+		Buttons[ (int32)EOculusTouchControllerButton::Touchpad ].EmulatedKey = (Hand == EControllerHand::Left) ? EKeys::MotionController_Left_Thumbstick.GetFName() : EKeys::MotionController_Right_Thumbstick.GetFName();
+
 		CapacitiveAxes[(int32)EOculusTouchCapacitiveAxes::Thumbstick].Axis = (Hand == EControllerHand::Left) ? FOculusKeyNames::OculusTouch_Left_Thumbstick : FOculusKeyNames::OculusTouch_Right_Thumbstick;
 		CapacitiveAxes[(int32)EOculusTouchCapacitiveAxes::Trigger].Axis = (Hand == EControllerHand::Left) ? FOculusKeyNames::OculusTouch_Left_Trigger : FOculusKeyNames::OculusTouch_Right_Trigger;
 		CapacitiveAxes[(int32)EOculusTouchCapacitiveAxes::XA].Axis = (Hand == EControllerHand::Left) ? FOculusKeyNames::OculusTouch_Left_FaceButton1 : FOculusKeyNames::OculusTouch_Right_FaceButton1;
diff --git a/Engine/Plugins/Runtime/OpenXR/Source/OpenXRHMD/Private/OpenXRHMD.h b/Engine/Plugins/Runtime/OpenXR/Source/OpenXRHMD/Private/OpenXRHMD.h
--- a/Engine/Plugins/Runtime/OpenXR/Source/OpenXRHMD/Private/OpenXRHMD.h
+++ b/Engine/Plugins/Runtime/OpenXR/Source/OpenXRHMD/Private/OpenXRHMD.h
@@ -181,6 +181,7 @@ public:
 	void FinishRendering();
 
 	OPENXRHMD_API int32 AddActionDevice(XrAction Action);
+	OPENXRHMD_API void ResetActionDevices();
 
 	FXRSwapChain* GetSwapchain() { return Swapchain.Get(); }
 	FXRSwapChain* GetDepthSwapchain() { return DepthSwapchain.Get(); }
diff --git a/Engine/Plugins/Runtime/OpenXR/Source/OpenXRInput/Private/OpenXRInput.cpp b/Engine/Plugins/Runtime/OpenXR/Source/OpenXRInput/Private/OpenXRInput.cpp
--- a/Engine/Plugins/Runtime/OpenXR/Source/OpenXRInput/Private/OpenXRInput.cpp
+++ b/Engine/Plugins/Runtime/OpenXR/Source/OpenXRInput/Private/OpenXRInput.cpp
@@ -3,32 +3,16 @@
 #include "OpenXRInput.h"
 #include "OpenXRHMD.h"
 #include "OpenXRHMDPrivate.h"
-#include "XRInputSettings.h"
 #include "UObject/UObjectIterator.h"
 #include "GameFramework/InputSettings.h"
 
 #if WITH_EDITOR
 #include "Editor/EditorEngine.h"
-#include "ISettingsModule.h"
-#include "ISettingsSection.h"
 #include "Editor.h"
 #endif
 
 #include <openxr/openxr.h>
 
-#define LOCTEXT_NAMESPACE "OpenXR"
-
-FSimpleMulticastDelegate UXRInputSettings::OnSuggestedBindingsChanged;
-
-#if WITH_EDITOR
-void UXRInputSettings::PostEditChangeChainProperty(struct FPropertyChangedChainEvent& PropertyChangedEvent)
-{
-	Super::PostEditChangeChainProperty(PropertyChangedEvent);
-
-	OnSuggestedBindingsChanged.Broadcast();
-}
-#endif
-
 FORCEINLINE XrPath GetPath(XrInstance Instance, const char* PathString)
 {
 	XrPath Path = XR_NULL_PATH;
@@ -37,10 +21,9 @@ FORCEINLINE XrPath GetPath(XrInstance Instance, const char* PathString)
 	return Path;
 }
 
-FORCEINLINE XrPath GetPath(XrInstance Instance, FString PathString)
+FORCEINLINE XrPath GetPath(XrInstance Instance, const FString& PathString)
 {
-	FTCHARToUTF8 EngineNameConverter(*PathString);
-	return GetPath(Instance, EngineNameConverter.Get());
+	return GetPath(Instance, (ANSICHAR*)StringCast<ANSICHAR>(*PathString).Get());
 }
 
 FORCEINLINE void FilterActionName(const char* InActionName, char* OutActionName)
@@ -93,43 +76,14 @@ void FOpenXRInputPlugin::StartupModule()
 	{
 		InputDevice = MakeShared<FOpenXRInput>(OpenXRHMD);
 	}
-
-#if WITH_EDITOR
-	ISettingsModule* SettingsModule = FModuleManager::GetModulePtr<ISettingsModule>("Settings");
-
-	// While this should usually be true, it's not guaranteed that the settings module will be loaded in the editor.
-	// UBT allows setting bBuildDeveloperTools to false while bBuildEditor can be true.
-	// The former option indirectly controls loading of the "Settings" module.
-	if (SettingsModule)
-	{
-		SettingsModule->RegisterSettings("Project", "Plugins", "OpenXR Input",
-			LOCTEXT("XRInputSettingsName", "OpenXR Input"),
-			LOCTEXT("XRInputSettingsDescription", "Configure input for OpenXR."),
-			GetMutableDefault<UXRInputSettings>()
-		);
-	}
-#endif
 }
 
-void FOpenXRInputPlugin::ShutdownModule()
-{
-	IOpenXRInputPlugin::ShutdownModule();
-
-#if WITH_EDITOR
-	ISettingsModule* SettingsModule = FModuleManager::GetModulePtr<ISettingsModule>("Settings");
-	if (SettingsModule)
-	{
-		SettingsModule->UnregisterSettings("Project", "Plugins", "OpenXR Input");
-	}
-#endif
-}
-
-FOpenXRInputPlugin::FOpenXRAction::FOpenXRAction(XrActionSet InActionSet, XrActionType InActionType, const FName& InName)
+FOpenXRInputPlugin::FOpenXRAction::FOpenXRAction(XrActionSet InActionSet, XrActionType InActionType, const FName& InName, const TArray<XrPath>& SubactionPaths)
 	: Set(InActionSet)
 	, Type(InActionType)
 	, Name(InName)
-	, Keys()
 	, Handle(XR_NULL_HANDLE)
+	, KeyMap()
 {
 	char ActionName[NAME_SIZE];
 	Name.GetPlainANSIString(ActionName);
@@ -139,30 +93,12 @@ FOpenXRInputPlugin::FOpenXRAction::FOpenXRAction(XrActionSet InActionSet, XrActi
 	Info.next = nullptr;
 	FilterActionName(ActionName, Info.actionName);
 	Info.actionType = Type;
-	Info.countSubactionPaths = 0;
-	Info.subactionPaths = nullptr;
+	Info.countSubactionPaths = SubactionPaths.Num();
+	Info.subactionPaths = SubactionPaths.GetData();
 	FCStringAnsi::Strcpy(Info.localizedActionName, XR_MAX_LOCALIZED_ACTION_NAME_SIZE, ActionName);
 	XR_ENSURE(xrCreateAction(Set, &Info, &Handle));
 }
 
-FOpenXRInputPlugin::FOpenXRAction::FOpenXRAction(XrActionSet InSet, const TArray<FInputActionKeyMapping>& InActionKeys)
-	: FOpenXRAction(InSet, XR_ACTION_TYPE_BOOLEAN_INPUT, InActionKeys[0].ActionName)
-{
-	for (auto Mapping : InActionKeys)
-	{
-		Keys.Add(Mapping.Key);
-	}
-}
-
-FOpenXRInputPlugin::FOpenXRAction::FOpenXRAction(XrActionSet InSet, const TArray<FInputAxisKeyMapping>& InAxisKeys)
-	: FOpenXRAction(InSet, XR_ACTION_TYPE_FLOAT_INPUT, InAxisKeys[0].AxisName)
-{
-	for (auto Mapping : InAxisKeys)
-	{
-		Keys.Add(Mapping.Key);
-	}
-}
-
 FOpenXRInputPlugin::FOpenXRController::FOpenXRController(FOpenXRHMD* HMD, XrActionSet InActionSet, const char* InName)
 	: ActionSet(InActionSet)
 	, Action(XR_NULL_HANDLE)
@@ -192,6 +128,12 @@ FOpenXRInputPlugin::FOpenXRController::FOpenXRController(FOpenXRHMD* HMD, XrActi
 	}
 }
 
+FOpenXRInputPlugin::FInteractionProfile::FInteractionProfile(XrPath InProfile)
+	: Path(InProfile)
+	, Bindings()
+{
+}
+
 FOpenXRInputPlugin::FOpenXRInput::FOpenXRInput(FOpenXRHMD* HMD)
 	: OpenXRHMD(HMD)
 	, ActionSets()
@@ -203,7 +145,7 @@ FOpenXRInputPlugin::FOpenXRInput::FOpenXRInput(FOpenXRHMD* HMD)
 	IModularFeatures::Get().RegisterModularFeature(GetModularFeatureName(), this);
 	check(OpenXRHMD);
 
-	InitActions();
+	BuildActions();
 }
 
 FOpenXRInputPlugin::FOpenXRInput::~FOpenXRInput()
@@ -211,10 +153,12 @@ FOpenXRInputPlugin::FOpenXRInput::~FOpenXRInput()
 	DestroyActions();
 }
 
-void FOpenXRInputPlugin::FOpenXRInput::InitActions()
+void FOpenXRInputPlugin::FOpenXRInput::BuildActions()
 {
 	if (bActionsBound)
+	{
 		return;
+	}
 
 	XrInstance Instance = OpenXRHMD->GetInstance();
 	check(Instance);
@@ -228,130 +172,180 @@ void FOpenXRInputPlugin::FOpenXRInput::InitActions()
 	FCStringAnsi::Strcpy(SetInfo.localizedActionSetName, XR_MAX_ACTION_SET_NAME_SIZE, "Unreal Engine 4");
 	XR_ENSURE(xrCreateActionSet(Instance, &SetInfo, &ActionSet));
 
-	const UInputSettings* InputSettings = GetDefault<UInputSettings>();
+	// Controller poses
+	OpenXRHMD->ResetActionDevices();
+	Controllers.Add(EControllerHand::Left, FOpenXRController(OpenXRHMD, ActionSet, "Left Controller"));
+	Controllers.Add(EControllerHand::Right, FOpenXRController(OpenXRHMD, ActionSet, "Right Controller"));
+
+	// Generate a map of all supported interaction profiles
+	TMap<FString, FInteractionProfile> Profiles;
+	Profiles.Add("Daydream", FInteractionProfile(GetPath(Instance, "/interaction_profiles/google/daydream_controller")));
+	Profiles.Add("Vive", FInteractionProfile(GetPath(Instance, "/interaction_profiles/htc/vive_controller")));
+	Profiles.Add("MixedReality", FInteractionProfile(GetPath(Instance, "/interaction_profiles/microsoft/motion_controller")));
+	Profiles.Add("OculusGo", FInteractionProfile(GetPath(Instance, "/interaction_profiles/oculus/go_controller")));
+	Profiles.Add("OculusTouch", FInteractionProfile(GetPath(Instance, "/interaction_profiles/oculus/touch_controller")));
+	Profiles.Add("ValveIndex", FInteractionProfile(GetPath(Instance, "/interaction_profiles/valve/index_controller")));
+
+	// Generate a list of the sub-action paths so we can query the left/right hand individually
+	SubactionPaths.Add(GetPath(Instance, "/user/hand/left"));
+	SubactionPaths.Add(GetPath(Instance, "/user/hand/right"));
+
+	auto InputSettings = GetMutableDefault<UInputSettings>();
 	if (InputSettings != nullptr)
 	{
 		TArray<FName> ActionNames;
 		InputSettings->GetActionNames(ActionNames);
-		for (const auto& ActionName : ActionNames)
+		for (const FName& ActionName : ActionNames)
 		{
+			FOpenXRAction Action(ActionSet, XR_ACTION_TYPE_BOOLEAN_INPUT, ActionName, SubactionPaths);
 			TArray<FInputActionKeyMapping> Mappings;
 			InputSettings->GetActionMappingByName(ActionName, Mappings);
-			Actions.Emplace(ActionSet, Mappings);
+
+			// If the developer didn't suggest any XR bindings for this action we won't expose it to the runtime
+			if (SuggestBindings(Instance, Action, Mappings, Profiles) > 0)
+			{
+				Actions.Add(Action);
+			}
+			else
+			{
+				XR_ENSURE(xrDestroyAction(Action.Handle));
+			}
 		}
 
 		TArray<FName> AxisNames;
 		InputSettings->GetAxisNames(AxisNames);
-		for (const auto& AxisName : AxisNames)
+		for (const FName& AxisName : AxisNames)
 		{
+			FOpenXRAction Action(ActionSet, XR_ACTION_TYPE_FLOAT_INPUT, AxisName, SubactionPaths);
 			TArray<FInputAxisKeyMapping> Mappings;
 			InputSettings->GetAxisMappingByName(AxisName, Mappings);
-			Actions.Emplace(ActionSet, Mappings);
+
+			// If the developer didn't suggest any XR bindings for this action we won't expose it to the runtime
+			if (SuggestBindings(Instance, Action, Mappings, Profiles) > 0)
+			{
+				Actions.Add(Action);
+			}
+			else
+			{
+				XR_ENSURE(xrDestroyAction(Action.Handle));
+			}
 		}
+
+		InputSettings->ForceRebuildKeymaps();
 	}
 
-	const UXRInputSettings* XRInputSettings = GetDefault<UXRInputSettings>();
+	for (TPair<FString, FInteractionProfile>& Pair : Profiles)
+	{
+		FInteractionProfile& Profile = Pair.Value;
 
-	// Controller poses
-	Controllers.Add(EControllerHand::Left, FOpenXRController(OpenXRHMD, ActionSet, "Left Controller"));
-	Controllers.Add(EControllerHand::Right, FOpenXRController(OpenXRHMD, ActionSet, "Right Controller"));
+		// Only suggest interaction profile bindings if the developer has provided bindings for them
+		if (Profile.Bindings.Num() > 0)
+		{
+			// Add the bindings for the controller pose and haptics
+			Profile.Bindings.Add(XrActionSuggestedBinding {
+				Controllers[EControllerHand::Left].Action, GetPath(Instance, "/user/hand/left/input/grip/pose")
+			});
+			Profile.Bindings.Add(XrActionSuggestedBinding {
+				Controllers[EControllerHand::Right].Action, GetPath(Instance, "/user/hand/right/input/grip/pose")
+			});
+			Profile.Bindings.Add(XrActionSuggestedBinding {
+				Controllers[EControllerHand::Left].VibrationAction, GetPath(Instance, "/user/hand/left/output/haptic")
+			});
+			Profile.Bindings.Add(XrActionSuggestedBinding {
+				Controllers[EControllerHand::Right].VibrationAction, GetPath(Instance, "/user/hand/right/output/haptic")
+			});
+
+			XrInteractionProfileSuggestedBinding InteractionProfile;
+			InteractionProfile.type = XR_TYPE_INTERACTION_PROFILE_SUGGESTED_BINDING;
+			InteractionProfile.next = nullptr;
+			InteractionProfile.interactionProfile = Profile.Path;
+			InteractionProfile.countSuggestedBindings = Profile.Bindings.Num();
+			InteractionProfile.suggestedBindings = Profile.Bindings.GetData();
+			XR_ENSURE(xrSuggestInteractionProfileBindings(Instance, &InteractionProfile));
+		}
+	}
 
-	SuggestedBindings(Instance, "/interaction_profiles/khr/simple_controller", XRInputSettings->SimpleController, XRInputSettings->SimpleBindings);
-	SuggestedBindings(Instance, "/interaction_profiles/google/daydream_controller", XRInputSettings->DaydreamController, XRInputSettings->DaydreamBindings);
-	SuggestedBindings(Instance, "/interaction_profiles/htc/vive_controller", XRInputSettings->ViveController, XRInputSettings->ViveBindings);
-	SuggestedBindings(Instance, "/interaction_profiles/htc/vive_pro", XRInputSettings->ViveProHeadset, XRInputSettings->ViveProBindings);
-	SuggestedBindings(Instance, "/interaction_profiles/microsoft/motion_controller", XRInputSettings->MixedRealityController, XRInputSettings->MixedRealityBindings);
-	SuggestedBindings(Instance, "/interaction_profiles/oculus/go_controller", XRInputSettings->OculusGoController, XRInputSettings->OculusGoBindings);
-	SuggestedBindings(Instance, "/interaction_profiles/oculus/touch_controller", XRInputSettings->OculusTouchController, XRInputSettings->OculusTouchBindings);
-	SuggestedBindings(Instance, "/interaction_profiles/valve/index_controller", XRInputSettings->ValveIndexController, XRInputSettings->ValveIndexBindings);
-
-	XrActiveActionSet ActiveSet;
-	ActiveSet.actionSet = ActionSet;
-	ActiveSet.subactionPath = XR_NULL_PATH;
-	ActionSets.Add(ActiveSet);
+	// Add an active set for each sub-action path so we can use the subaction paths later
+	// TODO: Runtimes already allow us to do the same by simply specifying "subactionPath"
+	// as XR_NULL_HANDLE, we're just being verbose for safety.
+	for (XrPath Subaction : SubactionPaths)
+	{
+		XrActiveActionSet ActiveSet;
+		ActiveSet.actionSet = ActionSet;
+		ActiveSet.subactionPath = Subaction;
+		ActionSets.Add(ActiveSet);
+	}
 }
 
 void FOpenXRInputPlugin::FOpenXRInput::DestroyActions()
 {
+	// Destroying an action set will also destroy all actions in the set
 	for (XrActiveActionSet& ActionSet : ActionSets)
 	{
-		XR_ENSURE(xrDestroyActionSet(ActionSet.actionSet));
+		xrDestroyActionSet(ActionSet.actionSet);
 	}
-	ActionSets.Empty();
-	Actions.Empty();
-	Controllers.Empty();
+
+	Actions.Reset();
+	Controllers.Reset();
+	SubactionPaths.Reset();
+	ActionSets.Reset();
 }
 
-void FOpenXRInputPlugin::FOpenXRInput::SuggestedBindings(XrInstance Instance, const char* Profile, const FXRInteractionProfileSettings& Settings, const TArray<FXRSuggestedBinding>& SuggestedBindings)
+template<typename T>
+int32 FOpenXRInputPlugin::FOpenXRInput::SuggestBindings(XrInstance Instance, FOpenXRAction& Action, const TArray<T>& Mappings, TMap<FString, FInteractionProfile>& Profiles)
 {
-	if (!Settings.Supported)
-	{
-		return;
-	}
+	int32 SuggestedBindings = 0;
 
-	TArray<XrActionSuggestedBinding> Bindings;
-	for (const auto& Suggestion : SuggestedBindings)
+	// Add suggested bindings for every mapping
+	for (const T& InputKey : Mappings)
 	{
-		if (Suggestion.Path.IsEmpty())
+		// Key names that are parseable into an OpenXR path have exeactly 4 components
+		TArray<FString> Tokens;
+		if (InputKey.Key.ToString().ParseIntoArray(Tokens, TEXT("_")) != 4)
 		{
 			continue;
 		}
 
-		const FKey& Key = Suggestion.Key;
-		const FOpenXRAction* Action = Actions.FindByPredicate([Key](FOpenXRAction& Mapping)
+		// Check if we support the profile specified in the key name
+		FInteractionProfile* Profile = Profiles.Find(Tokens[0]);
+		if (Profile)
 		{
-			int32 Index;
-			if (Mapping.Keys.Find(Key, Index))
+			// Parse the key name into an OpenXR interaction profile path
+			FString Path = "/user/hand/" + Tokens[1].ToLower();
+			XrPath TopLevel = GetPath(Instance, Path);
+
+			// Map this key to the correct subaction for this profile
+			// We'll use this later to trigger the correct key
+			TPair<XrPath, XrPath> Key(Profile->Path, TopLevel);
+			Action.KeyMap.Add(Key, InputKey.Key.GetFName());
+
+			// Add the input we want to query with grip being defined as "squeeze" in OpenXR
+			FString InputString = Tokens[2].ToLower();
+			if (InputString == "grip")
 			{
-				Mapping.Keys.Swap(0, Index);
-				return true;
+				InputString = "squeeze";
 			}
-			return false;
-		});
-
-		if (Action)
-		{
-			XrActionSuggestedBinding Binding;
-			Binding.action = Action->Handle;
-			Binding.binding = GetPath(Instance, Suggestion.Path);
-			Bindings.Add(Binding);
-		}
-	}
-
-	if (Settings.HasControllers)
-	{
-		XrActionSuggestedBinding LeftBinding, RightBinding;
-		LeftBinding.action = Controllers[EControllerHand::Left].Action;
-		RightBinding.action = Controllers[EControllerHand::Right].Action;
+			Path += "/input/" + InputString;
 
-		if (Settings.ControllerPose == AimPose)
-		{
-			LeftBinding.binding = GetPath(Instance, "/user/hand/left/input/aim/pose");
-			RightBinding.binding = GetPath(Instance, "/user/hand/right/input/aim/pose");
-		}
-		else
-		{
-			LeftBinding.binding = GetPath(Instance, "/user/hand/left/input/grip/pose");
-			RightBinding.binding = GetPath(Instance, "/user/hand/right/input/grip/pose");
-		}
-
-		Bindings.Add(LeftBinding);
-		Bindings.Add(RightBinding);
+			// Add the data we want to query, we'll skip this for trigger/squeeze "click" actions to allow
+			// certain profiles that don't have "click" data to threshold the "value" data instead
+			FString DataString = Tokens[3].ToLower();
+			if (DataString == "axis")
+			{
+				DataString = "value";
+			}
+			if ((InputString != "trigger" && InputString != "squeeze") || DataString != "click")
+			{
+				Path += "/" + DataString;
+			}
 
-		if (Settings.HasHaptics)
-		{
-			Bindings.Add(XrActionSuggestedBinding{ Controllers[EControllerHand::Left].VibrationAction, GetPath(Instance, "/user/hand/left/output/haptic") });
-			Bindings.Add(XrActionSuggestedBinding{ Controllers[EControllerHand::Right].VibrationAction, GetPath(Instance, "/user/hand/right/output/haptic") });
+			// Add the binding to the profile
+			Profile->Bindings.Add(XrActionSuggestedBinding{ Action.Handle, GetPath(Instance, Path) });
+			SuggestedBindings++;
 		}
 	}
 
-	XrInteractionProfileSuggestedBinding InteractionProfile;
-	InteractionProfile.type = XR_TYPE_INTERACTION_PROFILE_SUGGESTED_BINDING;
-	InteractionProfile.next = nullptr;
-	InteractionProfile.interactionProfile = GetPath(Instance, Profile);
-	InteractionProfile.countSuggestedBindings = Bindings.Num();
-	InteractionProfile.suggestedBindings = Bindings.GetData();
-	XR_ENSURE(xrSuggestInteractionProfileBindings(Instance, &InteractionProfile));
+	return SuggestedBindings;
 }
 
 void FOpenXRInputPlugin::FOpenXRInput::Tick(float DeltaTime)
@@ -402,51 +396,65 @@ void FOpenXRInputPlugin::FOpenXRInput::SendControllerEvents()
 
 	XrSession Session = OpenXRHMD->GetSession();
 
-	for (auto& Action : Actions)
+	for (XrPath Subaction : SubactionPaths)
 	{
-		XrActionStateGetInfo GetInfo;
-		GetInfo.type = XR_TYPE_ACTION_STATE_GET_INFO;
-		GetInfo.next = nullptr;
-		GetInfo.subactionPath = XR_NULL_PATH;
-		GetInfo.action = Action.Handle;
+		XrInteractionProfileState Profile;
+		Profile.type = XR_TYPE_INTERACTION_PROFILE_STATE;
+		Profile.next = nullptr;
+		XR_ENSURE(xrGetCurrentInteractionProfile(Session, Subaction, &Profile));
 
-		switch (Action.Type)
+		TPair<XrPath, XrPath> Key(Profile.interactionProfile, Subaction);
+		for (FOpenXRAction& Action : Actions)
 		{
-		case XR_ACTION_TYPE_BOOLEAN_INPUT:
-		{
-			XrActionStateBoolean State;
-			State.type = XR_TYPE_ACTION_STATE_BOOLEAN;
-			State.next = nullptr;
-			XrResult Result = xrGetActionStateBoolean(Session, &GetInfo, &State);
+			XrActionStateGetInfo GetInfo;
+			GetInfo.type = XR_TYPE_ACTION_STATE_GET_INFO;
+			GetInfo.next = nullptr;
+			GetInfo.subactionPath = Subaction;
+			GetInfo.action = Action.Handle;
+
+			FName* ActionKey = Action.KeyMap.Find(Key);
+			if (!ActionKey)
+			{
+				continue;
+			}
 
-			if (Result >= XR_SUCCESS && State.changedSinceLastSync)
+			switch (Action.Type)
 			{
-				if (State.currentState)
-				{
-					MessageHandler->OnControllerButtonPressed(Action.Keys[0].GetFName(), 0, /*IsRepeat =*/false);
-				}
-				else
+			case XR_ACTION_TYPE_BOOLEAN_INPUT:
+			{
+				XrActionStateBoolean State;
+				State.type = XR_TYPE_ACTION_STATE_BOOLEAN;
+				State.next = nullptr;
+				XrResult Result = xrGetActionStateBoolean(Session, &GetInfo, &State);
+				if (Result >= XR_SUCCESS && State.changedSinceLastSync)
 				{
-					MessageHandler->OnControllerButtonReleased(Action.Keys[0].GetFName(), 0, /*IsRepeat =*/false);
+					if (State.currentState)
+					{
+						MessageHandler->OnControllerButtonPressed(*ActionKey, 0, /*IsRepeat =*/false);
+					}
+					else
+					{
+						MessageHandler->OnControllerButtonReleased(*ActionKey, 0, /*IsRepeat =*/false);
+					}
 				}
 			}
-		}
-		break;
-		case XR_ACTION_TYPE_FLOAT_INPUT:
-		{
-			XrActionStateFloat State;
-			State.type = XR_TYPE_ACTION_STATE_FLOAT;
-			State.next = nullptr;
-			XrResult Result = xrGetActionStateFloat(Session, &GetInfo, &State);
-			if (Result >= XR_SUCCESS && State.changedSinceLastSync)
+			break;
+			case XR_ACTION_TYPE_FLOAT_INPUT:
 			{
-				MessageHandler->OnControllerAnalog(Action.Keys[0].GetFName(), 0, State.currentState);
+				XrActionStateFloat State;
+				State.type = XR_TYPE_ACTION_STATE_FLOAT;
+				State.next = nullptr;
+				XrResult Result = xrGetActionStateFloat(Session, &GetInfo, &State);
+				if (Result >= XR_SUCCESS && State.changedSinceLastSync)
+				{
+					MessageHandler->OnControllerAnalog(*ActionKey, 0, State.currentState);
+				}
 			}
-		}
-		break;
-		default:
-			// Other action types are currently unsupported.
 			break;
+			default:
+				// Other action types are currently unsupported.
+				break;
+			}
 		}
 	}
 }
@@ -454,10 +462,8 @@ void FOpenXRInputPlugin::FOpenXRInput::SendControllerEvents()
 void FOpenXRInputPlugin::FOpenXRInput::SetMessageHandler(const TSharedRef< FGenericApplicationMessageHandler >& InMessageHandler)
 {
 	MessageHandler = InMessageHandler;
-
-	UXRInputSettings::OnSuggestedBindingsChanged.AddSP(this, &FOpenXRInputPlugin::FOpenXRInput::InitActions);
 #if WITH_EDITOR
-	FEditorDelegates::OnActionAxisMappingsChanged.AddSP(this, &FOpenXRInputPlugin::FOpenXRInput::InitActions);
+	FEditorDelegates::OnActionAxisMappingsChanged.AddSP(this, &FOpenXRInputPlugin::FOpenXRInput::BuildActions);
 #endif
 }
 
@@ -583,5 +589,3 @@ float FOpenXRInputPlugin::FOpenXRInput::GetHapticAmplitudeScale() const
 {
 	return 1.0f;
 }
-
-#undef LOCTEXT_NAMESPACE
diff --git a/Engine/Plugins/Runtime/OpenXR/Source/OpenXRInput/Private/OpenXRInput.h b/Engine/Plugins/Runtime/OpenXR/Source/OpenXRInput/Private/OpenXRInput.h
--- a/Engine/Plugins/Runtime/OpenXR/Source/OpenXRInput/Private/OpenXRInput.h
+++ b/Engine/Plugins/Runtime/OpenXR/Source/OpenXRInput/Private/OpenXRInput.h
@@ -6,7 +6,6 @@
 #include "IOpenXRInputPlugin.h"
 #include "IInputDevice.h"
 #include "IHapticDevice.h"
-#include "XRInputSettings.h"
 
 #include <openxr/openxr.h>
 
@@ -22,13 +21,11 @@ public:
 		XrActionSet		Set;
 		XrActionType	Type;
 		FName			Name;
-		TArray<FKey>	Keys;
 		XrAction		Handle;
 
-		FOpenXRAction(XrActionSet InActionSet, XrActionType InActionType, const FName& InName);
+		TMap<TPair<XrPath, XrPath>, FName> KeyMap;
 
-		FOpenXRAction(XrActionSet InSet, const TArray<FInputActionKeyMapping>& InActionKeys);
-		FOpenXRAction(XrActionSet InSet, const TArray<FInputAxisKeyMapping>& InAxisKeys);
+		FOpenXRAction(XrActionSet InActionSet, XrActionType InActionType, const FName& InName, const TArray<XrPath>& SubactionPaths);
 	};
 
 	struct FOpenXRController
@@ -41,11 +38,16 @@ public:
 		FOpenXRController(FOpenXRHMD* HMD, XrActionSet InActionSet, const char* InName);
 	};
 
-	class FOpenXRInput :
-		public IInputDevice,
-		public FXRMotionControllerBase,
-		public IHapticDevice,
-		public TSharedFromThis<FOpenXRInput>
+	struct FInteractionProfile
+	{
+	public:
+		XrPath Path;
+		TArray<XrActionSuggestedBinding> Bindings;
+
+		FInteractionProfile(XrPath InProfile);
+	};
+
+	class FOpenXRInput : public IInputDevice, public FXRMotionControllerBase, public IHapticDevice, public TSharedFromThis<FOpenXRInput>
 	{
 	public:
 		FOpenXRInput(FOpenXRHMD* HMD);
@@ -77,14 +79,16 @@ public:
 		FOpenXRHMD* OpenXRHMD;
 
 		TArray<XrActiveActionSet> ActionSets;
+		TArray<XrPath> SubactionPaths;
 		TArray<FOpenXRAction> Actions;
 		TMap<EControllerHand, FOpenXRController> Controllers;
 
 		bool bActionsBound;
-
-		void InitActions();
+		void BuildActions();
 		void DestroyActions();
-		void SuggestedBindings(XrInstance Instance, const char* Profile, const FXRInteractionProfileSettings& Settings, const TArray<FXRSuggestedBinding>& SuggestedBindings);
+
+		template<typename T>
+		int32 SuggestBindings(XrInstance Instance, FOpenXRAction& Action, const TArray<T>& Mappings, TMap<FString, FInteractionProfile>& Profiles);
 
 		/** handler to send all messages to */
 		TSharedRef<FGenericApplicationMessageHandler> MessageHandler;
@@ -96,7 +100,6 @@ public:
 	FOpenXRHMD* GetOpenXRHMD() const;
 
 	virtual void StartupModule() override;
-	virtual void ShutdownModule() override;
 	virtual TSharedPtr< class IInputDevice > CreateInputDevice(const TSharedRef< FGenericApplicationMessageHandler >& InMessageHandler) override;
 
 private:
diff --git a/Engine/Plugins/Runtime/Steam/SteamVR/Source/SteamVREditor/Private/AnimGraphNode_SteamVRInputAnimPose.cpp b/Engine/Plugins/Runtime/Steam/SteamVR/Source/SteamVREditor/Private/AnimGraphNode_SteamVRInputAnimPose.cpp
new file mode 100644
--- /dev/null
+++ b/Engine/Plugins/Runtime/Steam/SteamVR/Source/SteamVREditor/Private/AnimGraphNode_SteamVRInputAnimPose.cpp
@@ -0,0 +1,65 @@
+/*
+Copyright 2019 Valve Corporation under https://opensource.org/licenses/BSD-3-Clause
+
+Redistribution and use in source and binary forms, with or without modification,
+are permitted provided that the following conditions are met:
+
+1. Redistributions of source code must retain the above copyright notice, this
+   list of conditions and the following disclaimer.
+
+2. Redistributions in binary form must reproduce the above copyright notice,
+   this list of conditions and the following disclaimer in the documentation
+   and/or other materials provided with the distribution.
+
+3. Neither the name of the copyright holder nor the names of its contributors
+   may be used to endorse or promote products derived from this software
+   without specific prior written permission.
+
+THIS SOFTWARE IS PROVIDED BY THE COPYRIGHT HOLDERS AND CONTRIBUTORS "AS IS" AND
+ANY EXPRESS OR IMPLIED WARRANTIES, INCLUDING, BUT NOT LIMITED TO, THE IMPLIED
+WARRANTIES OF MERCHANTABILITY AND FITNESS FOR A PARTICULAR PURPOSE ARE DISCLAIMED.
+IN NO EVENT SHALL THE COPYRIGHT HOLDER OR CONTRIBUTORS BE LIABLE FOR ANY DIRECT,
+INDIRECT, INCIDENTAL, SPECIAL, EXEMPLARY, OR CONSEQUENTIAL DAMAGES (INCLUDING,
+BUT NOT LIMITED TO, PROCUREMENT OF SUBSTITUTE GOODS OR SERVICES; LOSS OF USE, DATA,
+OR PROFITS; OR BUSINESS INTERRUPTION) HOWEVER CAUSED AND ON ANY THEORY OF LIABILITY,
+WHETHER IN CONTRACT, STRICT LIABILITY, OR TORT (INCLUDING NEGLIGENCE OR OTHERWISE)
+ARISING IN ANY WAY OUT OF THE USE OF THIS SOFTWARE, EVEN IF ADVISED OF THE
+POSSIBILITY OF SUCH DAMAGE.
+*/
+
+
+#include "AnimGraphNode_SteamVRInputAnimPose.h"
+#include "SteamVREditor.h"
+
+#define LOCTEXT_NAMESPACE "SteamVRInputAnimNode"
+
+UAnimGraphNode_SteamVRInputAnimPose::UAnimGraphNode_SteamVRInputAnimPose(const FObjectInitializer& ObjectInitializer)
+	: Super(ObjectInitializer)
+{
+}
+
+// Node Color
+FLinearColor UAnimGraphNode_SteamVRInputAnimPose::GetNodeTitleColor() const 
+{ 
+	return FLinearColor(0.f, 0.f, 255.f, 1.f);
+}
+
+// Node Category
+FText UAnimGraphNode_SteamVRInputAnimPose::GetMenuCategory() const
+{
+	return LOCTEXT("NodeCategory", "SteamVR Input");
+}
+
+// Node Title
+FText UAnimGraphNode_SteamVRInputAnimPose::GetNodeTitle(ENodeTitleType::Type TitleType) const
+{
+	return LOCTEXT("NodeTitle", "SteamVR Skeletal Anim Pose");
+}
+
+// Node Tooltip
+FText UAnimGraphNode_SteamVRInputAnimPose::GetTooltipText() const
+{
+	return LOCTEXT("NodeTooltip", "Retrieves the current pose from the SteamVR Skeletal Input API");
+}
+
+#undef LOCTEXT_NAMESPACE
\ No newline at end of file
