Merging CL 11058926 from //UE4/Dev-VR to //UE4/Release-4.24

Fix HMD device space being unintentionally freed when the VR preview is closed

diff --git a/Engine/Plugins/Runtime/OpenXR/Source/OpenXRHMD/Private/OpenXRHMD.cpp b/Engine/Plugins/Runtime/OpenXR/Source/OpenXRHMD/Private/OpenXRHMD.cpp
--- a/Engine/Plugins/Runtime/OpenXR/Source/OpenXRHMD/Private/OpenXRHMD.cpp
+++ b/Engine/Plugins/Runtime/OpenXR/Source/OpenXRHMD/Private/OpenXRHMD.cpp
@@ -562,7 +562,7 @@ bool FOpenXRHMD::EnumerateTrackedDevices(TArray<int32>& OutDevices, EXRTrackedDe
 	}
 	if (Type == EXRTrackedDeviceType::Any || Type == EXRTrackedDeviceType::Controller)
 	{
-		for (int32 i = 0; i < ActionSpaces.Num(); i++)
+		for (int32 i = 0; i < DeviceSpaces.Num(); i++)
 		{
 			OutDevices.Add(i);
 		}
@@ -584,14 +584,14 @@ bool FOpenXRHMD::GetCurrentPose(int32 DeviceId, FQuat& CurrentOrientation, FVect
 	CurrentOrientation = FQuat::Identity;
 	CurrentPosition = FVector::ZeroVector;
 
-	if (!ActionSpaces.IsValidIndex(DeviceId) || !ActionSpaces[DeviceId].Space || FrameState.predictedDisplayTime <= 0)
+	if (!DeviceSpaces.IsValidIndex(DeviceId) || !DeviceSpaces[DeviceId].Space || FrameState.predictedDisplayTime <= 0)
 	{
 		return false;
 	}
 
 	XrSpaceLocation Location = {};
 	Location.type = XR_TYPE_SPACE_LOCATION;
-	XrResult Result = xrLocateSpace(ActionSpaces[DeviceId].Space, GetTrackingSpace(), FrameState.predictedDisplayTime, &Location);
+	XrResult Result = xrLocateSpace(DeviceSpaces[DeviceId].Space, GetTrackingSpace(), FrameState.predictedDisplayTime, &Location);
 	if (!XR_ENSURE(Result))
 	{
 		return false;
@@ -886,7 +886,8 @@ FOpenXRHMD::FOpenXRHMD(const FAutoRegister& AutoRegister, XrInstance InInstance,
 		View.pose = ToXrPose(FTransform::Identity);
 	}
 
-	ResetActionDevices();
+	// Add a device space for the HMD without an action handle and ensure it has the correct index
+	ensure(DeviceSpaces.Emplace(XR_NULL_HANDLE) == HMDDeviceId);
 }
 
 FOpenXRHMD::~FOpenXRHMD()
@@ -1043,7 +1044,7 @@ bool FOpenXRHMD::OnStereoStartup()
 	SpaceInfo.referenceSpaceType = XR_REFERENCE_SPACE_TYPE_VIEW;
 	SpaceInfo.poseInReferenceSpace = ToXrPose(FTransform::Identity);
 	XR_ENSURE(xrCreateReferenceSpace(Session, &SpaceInfo, &HmdSpace));
-	ActionSpaces[HMDDeviceId].Space = HmdSpace;
+	DeviceSpaces[HMDDeviceId].Space = HmdSpace;
 
 	ensure(Spaces.Contains(XR_REFERENCE_SPACE_TYPE_LOCAL));
 	SpaceInfo.referenceSpaceType = XR_REFERENCE_SPACE_TYPE_LOCAL;
@@ -1062,9 +1063,9 @@ bool FOpenXRHMD::OnStereoStartup()
 	}
 
 	// Create action spaces for all devices
-	for (auto& ActionSpace : ActionSpaces)
+	for (auto& DeviceSpace : DeviceSpaces)
 	{
-		ActionSpace.CreateSpace(Session);
+		DeviceSpace.CreateSpace(Session);
 	}
 
 	RenderBridge->SetOpenXRHMD(this);
@@ -1103,10 +1104,10 @@ void FOpenXRHMD::CloseSession()
 {
 	if (Session != XR_NULL_HANDLE)
 	{
-		// Clear up action spaces
-		for (auto& ActionSpace : ActionSpaces)
+		// Clear up device spaces
+		for (auto& DeviceSpace : DeviceSpaces)
 		{
-			ActionSpace.DestroySpace();
+			DeviceSpace.DestroySpace();
 		}
 
 		// Close the session now we're allowed to.
@@ -1133,10 +1134,10 @@ void FOpenXRHMD::CloseSession()
 
 int32 FOpenXRHMD::AddActionDevice(XrAction Action)
 {
-	int32 DeviceId = ActionSpaces.Emplace(Action);
+	int32 DeviceId = DeviceSpaces.Emplace(Action);
 	if (Session)
 	{
-		ActionSpaces[DeviceId].CreateSpace(Session);
+		DeviceSpaces[DeviceId].CreateSpace(Session);
 	}
 
 	return DeviceId;
@@ -1144,10 +1145,11 @@ int32 FOpenXRHMD::AddActionDevice(XrAction Action)
 
 void FOpenXRHMD::ResetActionDevices()
 {
-	ActionSpaces.Reset();
-
-	// The HMD device does not have an action associated with it
-	ensure(ActionSpaces.Emplace(XR_NULL_HANDLE) == HMDDeviceId);
+	// Index 0 is HMDDeviceId and is preserved. The remaining are action devices.
+	if (DeviceSpaces.Num() > 0)
+	{
+		DeviceSpaces.RemoveAt(HMDDeviceId + 1, DeviceSpaces.Num() - 1);
+	}
 }
 
 bool FOpenXRHMD::IsInitialized() const
@@ -1365,7 +1367,7 @@ void FOpenXRHMD::OnBeginRendering_GameThread()
 	ViewInfo.type = XR_TYPE_VIEW_LOCATE_INFO;
 	ViewInfo.next = nullptr;
 	ViewInfo.viewConfigurationType = XR_VIEW_CONFIGURATION_TYPE_PRIMARY_STEREO;
-	ViewInfo.space = ActionSpaces[HMDDeviceId].Space;
+	ViewInfo.space = DeviceSpaces[HMDDeviceId].Space;
 	ViewInfo.displayTime = FrameState.predictedDisplayTime;
 	XR_ENSURE(xrLocateViews(Session, &ViewInfo, &ViewState, 0, &ViewCount, nullptr));
 	Views.SetNum(ViewCount);
@@ -1663,13 +1665,13 @@ void FOpenXRHMD::DrawVisibleAreaMesh_RenderThread(class FRHICommandList& RHICmdL
 // OpenXR Action Space Implementation
 //---------------------------------------------------
 
-FOpenXRHMD::FActionSpace::FActionSpace(XrAction InAction)
+FOpenXRHMD::FDeviceSpace::FDeviceSpace(XrAction InAction)
 	: Action(InAction)
 	, Space(XR_NULL_HANDLE)
 {
 }
 
-bool FOpenXRHMD::FActionSpace::CreateSpace(XrSession InSession)
+bool FOpenXRHMD::FDeviceSpace::CreateSpace(XrSession InSession)
 {
 	if (Action == XR_NULL_HANDLE || Space != XR_NULL_HANDLE)
 	{
@@ -1685,7 +1687,7 @@ bool FOpenXRHMD::FActionSpace::CreateSpace(XrSession InSession)
 	return XR_ENSURE(xrCreateActionSpace(InSession, &ActionSpaceInfo, &Space));
 }
 
-void FOpenXRHMD::FActionSpace::DestroySpace()
+void FOpenXRHMD::FDeviceSpace::DestroySpace()
 {
 	if (Space)
 	{
diff --git a/Engine/Plugins/Runtime/OpenXR/Source/OpenXRHMD/Private/OpenXRHMD.h b/Engine/Plugins/Runtime/OpenXR/Source/OpenXRHMD/Private/OpenXRHMD.h
--- a/Engine/Plugins/Runtime/OpenXR/Source/OpenXRHMD/Private/OpenXRHMD.h
+++ b/Engine/Plugins/Runtime/OpenXR/Source/OpenXRHMD/Private/OpenXRHMD.h
@@ -26,14 +26,15 @@ class FOpenXRRenderBridge;
 class FOpenXRHMD : public FHeadMountedDisplayBase, public FXRRenderTargetManager, public FSceneViewExtensionBase
 {
 public:
-	class FActionSpace
+	class FDeviceSpace
 	{
 	public:
-		FActionSpace(XrAction InAction);
+		FDeviceSpace(XrAction InAction);
 
 		bool CreateSpace(XrSession InSession);
 		void DestroySpace();
 
+		XrReferenceSpaceType ReferenceSpaceType;
 		XrAction Action;
 		XrSpace Space;
 	};
@@ -225,7 +226,7 @@ private:
 	TArray<XrCompositionLayerProjectionView> ViewsRHI;
 	TArray<XrCompositionLayerDepthInfoKHR> DepthLayersRHI;
 
-	TArray<FActionSpace>	ActionSpaces;
+	TArray<FDeviceSpace>	DeviceSpaces;
 
 	TRefCountPtr<FOpenXRRenderBridge> RenderBridge;
 	IRendererModule*		RendererModule;
