Fix for warning after upgrading to PVS-Studio 7.0.4:

V781: The value of the 'i' index is checked after it was used. Perhaps there is a mistake in program logic.

#rnx

diff --git a/Engine/Plugins/Runtime/OpenXR/Source/OpenXRInput/Private/OpenXRInput.cpp b/Engine/Plugins/Runtime/OpenXR/Source/OpenXRInput/Private/OpenXRInput.cpp
--- a/Engine/Plugins/Runtime/OpenXR/Source/OpenXRInput/Private/OpenXRInput.cpp
+++ b/Engine/Plugins/Runtime/OpenXR/Source/OpenXRInput/Private/OpenXRInput.cpp
@@ -30,7 +30,7 @@ FORCEINLINE void FilterActionName(const char* InActionName, char* OutActionName)
 {
 	// Ensure the action name is a well-formed path
 	size_t i;
-	for (i = 0; InActionName[i] != '\0' && i < XR_MAX_ACTION_NAME_SIZE - 1; i++)
+	for (i = 0; i < XR_MAX_ACTION_NAME_SIZE - 1 && InActionName[i] != '\0'; i++)
 	{
 		unsigned char c = InActionName[i];
 		OutActionName[i] = (c == ' ') ? '-' : isalnum(c) ? tolower(c) : '_';
