Remove redundant swapchain re-creation when switching modes (e.g. VRMode/PIE and back).

diff --git a/Engine/Plugins/Runtime/OpenXR/Source/OpenXRHMD/Private/OpenXRHMD.cpp b/Engine/Plugins/Runtime/OpenXR/Source/OpenXRHMD/Private/OpenXRHMD.cpp
--- a/Engine/Plugins/Runtime/OpenXR/Source/OpenXRHMD/Private/OpenXRHMD.cpp
+++ b/Engine/Plugins/Runtime/OpenXR/Source/OpenXRHMD/Private/OpenXRHMD.cpp
@@ -833,6 +833,8 @@ FOpenXRHMD::FOpenXRHMD(const FAutoRegister& AutoRegister, XrInstance InInstance,
 	, SelectedViewConfigurationType(XR_VIEW_CONFIGURATION_TYPE_PRIMARY_STEREO)
 	, RenderBridge(InRenderBridge)
 	, RendererModule(nullptr)
+	, LastRequestedSwapchainFormat(0)
+	, LastRequestedDepthSwapchainFormat(0)
 {
 	FrameState.type = XR_TYPE_FRAME_STATE;
 	FrameState.next = nullptr;
@@ -1206,14 +1208,18 @@ bool FOpenXRHMD::AllocateRenderTargetTexture(uint32 Index, uint32 SizeX, uint32
 	// We need to ensure we can sample from the texture in CopyTexture
 	Flags |= TexCreate_ShaderResource;
 
-	Swapchain = RenderBridge->CreateSwapchain(Session, Format, SizeX, SizeY, NumMips, NumSamples, Flags, TargetableTextureFlags);
-	if (!Swapchain)
+	if (Swapchain == nullptr || Format != LastRequestedSwapchainFormat || Swapchain->GetTexture2D()->GetSizeX() != SizeX || Swapchain->GetTexture2D()->GetSizeY() != SizeY)
 	{
-		return false;
+		Swapchain = RenderBridge->CreateSwapchain(Session, Format, SizeX, SizeY, NumMips, NumSamples, Flags, TargetableTextureFlags);
+		if (!Swapchain)
+		{
+			return false;
+		}
 	}
-	
+
 	// Grab the presentation texture out of the swapchain.
 	OutTargetableTexture = OutShaderResourceTexture = (FTexture2DRHIRef&)Swapchain->GetTextureRef();
+	LastRequestedSwapchainFormat = Format;
 
 	bNeedReAllocatedDepth = bDepthExtensionSupported;
 
@@ -1228,15 +1234,19 @@ bool FOpenXRHMD::AllocateDepthTexture(uint32 Index, uint32 SizeX, uint32 SizeY,
 		return false;
 	}
 
-	DepthSwapchain = RenderBridge->CreateSwapchain(Session, PF_DepthStencil, SizeX, SizeY, FMath::Max(NumMips, 1u), NumSamples, 0, TexCreate_DepthStencilTargetable);
-	if (!DepthSwapchain)
+	if (DepthSwapchain == nullptr || Format != LastRequestedDepthSwapchainFormat || DepthSwapchain->GetTexture2D()->GetSizeX() != SizeX || DepthSwapchain->GetTexture2D()->GetSizeY() != SizeY)
 	{
-		return false;
+		DepthSwapchain = RenderBridge->CreateSwapchain(Session, PF_DepthStencil, SizeX, SizeY, FMath::Max(NumMips, 1u), NumSamples, 0, TexCreate_DepthStencilTargetable);
+		if (!DepthSwapchain)
+		{
+			return false;
+		}
 	}
 
 	bNeedReAllocatedDepth = false;
 
 	OutTargetableTexture = OutShaderResourceTexture = (FTexture2DRHIRef&)DepthSwapchain->GetTextureRef();
+	LastRequestedDepthSwapchainFormat = Format;
 
 	return true;
 }
diff --git a/Engine/Plugins/Runtime/OpenXR/Source/OpenXRHMD/Private/OpenXRHMD.h b/Engine/Plugins/Runtime/OpenXR/Source/OpenXRHMD/Private/OpenXRHMD.h
--- a/Engine/Plugins/Runtime/OpenXR/Source/OpenXRHMD/Private/OpenXRHMD.h
+++ b/Engine/Plugins/Runtime/OpenXR/Source/OpenXRHMD/Private/OpenXRHMD.h
@@ -232,6 +232,8 @@ private:
 
 	FXRSwapChainPtr			Swapchain;
 	FXRSwapChainPtr			DepthSwapchain;
+	uint8					LastRequestedSwapchainFormat;
+	uint8					LastRequestedDepthSwapchainFormat;
 
 	TArray<FHMDViewMesh>	HiddenAreaMeshes;
 	TArray<FHMDViewMesh>	VisibleAreaMeshes;
diff --git a/Engine/Plugins/Runtime/OpenXR/Source/OpenXRHMD/Private/OpenXRHMD_RenderBridge.cpp b/Engine/Plugins/Runtime/OpenXR/Source/OpenXRHMD/Private/OpenXRHMD_RenderBridge.cpp
--- a/Engine/Plugins/Runtime/OpenXR/Source/OpenXRHMD/Private/OpenXRHMD_RenderBridge.cpp
+++ b/Engine/Plugins/Runtime/OpenXR/Source/OpenXRHMD/Private/OpenXRHMD_RenderBridge.cpp
@@ -42,7 +42,7 @@ public:
 		return &Binding;
 	}
 
-	virtual FXRSwapChainPtr CreateSwapchain(XrSession InSession, uint8 Format, uint32 SizeX, uint32 SizeY, uint32 NumMips, uint32 NumSamples, uint32 Flags, uint32 TargetableTextureFlags)
+	virtual FXRSwapChainPtr CreateSwapchain(XrSession InSession, uint8 Format, uint32 SizeX, uint32 SizeY, uint32 NumMips, uint32 NumSamples, uint32 Flags, uint32 TargetableTextureFlags) override final
 	{
 		return CreateSwapchain_D3D11(InSession, Format, SizeX, SizeY, NumMips, NumSamples, Flags, TargetableTextureFlags);
 	}
@@ -78,7 +78,7 @@ public:
 		return &Binding;
 	}
 
-	virtual FXRSwapChainPtr CreateSwapchain(XrSession InSession, uint8 Format, uint32 SizeX, uint32 SizeY, uint32 NumMips, uint32 NumSamples, uint32 Flags, uint32 TargetableTextureFlags)
+	virtual FXRSwapChainPtr CreateSwapchain(XrSession InSession, uint8 Format, uint32 SizeX, uint32 SizeY, uint32 NumMips, uint32 NumSamples, uint32 Flags, uint32 TargetableTextureFlags) override final
 	{
 		return CreateSwapchain_D3D12(InSession, Format, SizeX, SizeY, NumMips, NumSamples, Flags, TargetableTextureFlags);
 	}
@@ -128,7 +128,7 @@ public:
 		return nullptr;
 	}
 
-	virtual FXRSwapChainPtr CreateSwapchain(XrSession InSession, uint8 Format, uint32 SizeX, uint32 SizeY, uint32 NumMips, uint32 NumSamples, uint32 Flags, uint32 TargetableTextureFlags)
+	virtual FXRSwapChainPtr CreateSwapchain(XrSession InSession, uint8 Format, uint32 SizeX, uint32 SizeY, uint32 NumMips, uint32 NumSamples, uint32 Flags, uint32 TargetableTextureFlags) override final
 	{
 		return CreateSwapchain_OpenGL(InSession, Format, SizeX, SizeY, NumMips, NumSamples, Flags, TargetableTextureFlags);
 	}
@@ -210,7 +210,7 @@ public:
 		return AdapterLuid;
 	}
 
-	virtual FXRSwapChainPtr CreateSwapchain(XrSession InSession, uint8 Format, uint32 SizeX, uint32 SizeY, uint32 NumMips, uint32 NumSamples, uint32 Flags, uint32 TargetableTextureFlags)
+	virtual FXRSwapChainPtr CreateSwapchain(XrSession InSession, uint8 Format, uint32 SizeX, uint32 SizeY, uint32 NumMips, uint32 NumSamples, uint32 Flags, uint32 TargetableTextureFlags) override final
 	{
 		return CreateSwapchain_Vulkan(InSession, Format, SizeX, SizeY, NumMips, NumSamples, Flags, TargetableTextureFlags);
 	}
diff --git a/Engine/Plugins/Runtime/OpenXR/Source/OpenXRHMD/Private/OpenXRHMD_Swapchain.cpp b/Engine/Plugins/Runtime/OpenXR/Source/OpenXRHMD/Private/OpenXRHMD_Swapchain.cpp
--- a/Engine/Plugins/Runtime/OpenXR/Source/OpenXRHMD/Private/OpenXRHMD_Swapchain.cpp
+++ b/Engine/Plugins/Runtime/OpenXR/Source/OpenXRHMD/Private/OpenXRHMD_Swapchain.cpp
@@ -196,11 +196,11 @@ FXRSwapChainPtr CreateSwapchain_D3D11(XrSession InSession, uint8 Format, uint32
 	TArray<FTextureRHIRef> TextureChain;
 	// @todo: Once things settle down, the chain target will be created below in the CreateXRSwapChain call, via an RHI "CreateAliasedTexture" call.
 	TArray<XrSwapchainImageD3D11KHR> Images = EnumerateImages<XrSwapchainImageD3D11KHR>(Swapchain, XR_TYPE_SWAPCHAIN_IMAGE_D3D11_KHR);
-	FTextureRHIRef ChainTarget = static_cast<FTextureRHIRef>(DynamicRHI->RHICreateTexture2DFromResource(GPixelFormats[Format].UnrealFormat, TargetableTextureFlags, bDepthStencil ? FClearValueBinding::DepthFar : FClearValueBinding::Black, Images[0].texture));
 	for (const auto& Image : Images)
 	{
 		TextureChain.Add(static_cast<FTextureRHIRef>(DynamicRHI->RHICreateTexture2DFromResource(GPixelFormats[Format].UnrealFormat, TargetableTextureFlags, bDepthStencil ? FClearValueBinding::DepthFar : FClearValueBinding::Black, Image.texture)));
 	}
+	FTextureRHIRef ChainTarget = static_cast<FTextureRHIRef>(GDynamicRHI->RHICreateAliasedTexture((FTextureRHIRef&)TextureChain[0]));
 
 	return CreateXRSwapChain<FOpenXRSwapchain>(MoveTemp(TextureChain), ChainTarget, Swapchain);
 }
@@ -236,11 +236,11 @@ FXRSwapChainPtr CreateSwapchain_D3D12(XrSession InSession, uint8 Format, uint32
 	TArray<FTextureRHIRef> TextureChain;
 	// @todo: Once things settle down, the chain target will be created below in the CreateXRSwapChain call, via an RHI "CreateAliasedTexture" call.
 	TArray<XrSwapchainImageD3D12KHR> Images = EnumerateImages<XrSwapchainImageD3D12KHR>(Swapchain, XR_TYPE_SWAPCHAIN_IMAGE_D3D12_KHR);
-	FTextureRHIRef ChainTarget = static_cast<FTextureRHIRef>(DynamicRHI->RHICreateTexture2DFromResource(GPixelFormats[Format].UnrealFormat, TargetableTextureFlags, FClearValueBinding::Black, Images[0].texture));
 	for (const auto& Image : Images)
 	{
 		TextureChain.Add(static_cast<FTextureRHIRef>(DynamicRHI->RHICreateTexture2DFromResource(GPixelFormats[Format].UnrealFormat, TargetableTextureFlags, FClearValueBinding::Black, Image.texture)));
 	}
+	FTextureRHIRef ChainTarget = static_cast<FTextureRHIRef>(GDynamicRHI->RHICreateAliasedTexture((FTextureRHIRef&)TextureChain[0]));
 
 	return CreateXRSwapChain<FOpenXRSwapchain>(MoveTemp(TextureChain), ChainTarget, Swapchain);
 }
@@ -265,11 +265,11 @@ FXRSwapChainPtr CreateSwapchain_OpenGL(XrSession InSession, uint8 Format, uint32
 	FOpenGLDynamicRHI* DynamicRHI = static_cast<FOpenGLDynamicRHI*>(GDynamicRHI);
 	// @todo: Once things settle down, the chain target will be created below in the CreateXRSwapChain call, via an RHI "CreateAliasedTexture" call.
 	TArray<XrSwapchainImageOpenGLKHR> Images = EnumerateImages<XrSwapchainImageOpenGLKHR>(Swapchain, XR_TYPE_SWAPCHAIN_IMAGE_OPENGL_KHR);
-	FTextureRHIRef ChainTarget = static_cast<FTextureRHIRef>(DynamicRHI->RHICreateTexture2DFromResource(GPixelFormats[Format].UnrealFormat, SizeX, SizeY, NumMips, NumSamples, 1, FClearValueBinding::Black, Images[0].image, Flags));
 	for (const auto& Image : Images)
 	{
 		TextureChain.Add(static_cast<FTextureRHIRef>(DynamicRHI->RHICreateTexture2DFromResource(GPixelFormats[Format].UnrealFormat, SizeX, SizeY, NumMips, NumSamples, 1, FClearValueBinding::Black, Image.image, Flags)));
 	}
+	FTextureRHIRef ChainTarget = static_cast<FTextureRHIRef>(GDynamicRHI->RHICreateAliasedTexture((FTextureRHIRef&)TextureChain[0]));
 
 	return CreateXRSwapChain<FOpenXRSwapchain>(MoveTemp(TextureChain), ChainTarget, Swapchain);
 }
@@ -299,11 +299,11 @@ FXRSwapChainPtr CreateSwapchain_Vulkan(XrSession InSession, uint8 Format, uint32
 	FVulkanDynamicRHI* DynamicRHI = static_cast<FVulkanDynamicRHI*>(GDynamicRHI);
 	// @todo: Once things settle down, the chain target will be created below in the CreateXRSwapChain call, via an RHI "CreateAliasedTexture" call.
 	TArray<XrSwapchainImageVulkanKHR> Images = EnumerateImages<XrSwapchainImageVulkanKHR>(Swapchain, XR_TYPE_SWAPCHAIN_IMAGE_VULKAN_KHR);
-	FTextureRHIRef ChainTarget = static_cast<FTextureRHIRef>(DynamicRHI->RHICreateTexture2DFromResource(GPixelFormats[Format].UnrealFormat, SizeX, SizeY, NumMips, NumSamples, Images[0].image, Flags));
 	for (const auto& Image : Images)
 	{
 		TextureChain.Add(static_cast<FTextureRHIRef>(DynamicRHI->RHICreateTexture2DFromResource(GPixelFormats[Format].UnrealFormat, SizeX, SizeY, NumMips, NumSamples, Image.image, Flags)));
 	}
+	FTextureRHIRef ChainTarget = static_cast<FTextureRHIRef>(GDynamicRHI->RHICreateAliasedTexture((FTextureRHIRef&)TextureChain[0]));
 
 	return CreateXRSwapChain<FOpenXRSwapchain>(MoveTemp(TextureChain), ChainTarget, Swapchain);
 }
diff --git a/Engine/Plugins/Runtime/Steam/SteamVR/Source/SteamVR/Private/SteamVRHMD.cpp b/Engine/Plugins/Runtime/Steam/SteamVR/Source/SteamVR/Private/SteamVRHMD.cpp
--- a/Engine/Plugins/Runtime/Steam/SteamVR/Source/SteamVR/Private/SteamVRHMD.cpp
+++ b/Engine/Plugins/Runtime/Steam/SteamVR/Source/SteamVR/Private/SteamVRHMD.cpp
@@ -1608,6 +1608,13 @@ bool FSteamVRHMD::AllocateRenderTargetTexture(uint32 Index, uint32 SizeX, uint32
 	MetalBridge* MetalBridgePtr = (MetalBridge*)pBridge.GetReference();
 #endif
 
+	if (pBridge != nullptr && pBridge->GetSwapChain() != nullptr && pBridge->GetSwapChain()->GetTexture2D() != nullptr && pBridge->GetSwapChain()->GetTexture2D()->GetSizeX() == SizeX && pBridge->GetSwapChain()->GetTexture2D()->GetSizeY() == SizeY)
+	{
+		OutTargetableTexture = (FTexture2DRHIRef&)pBridge->GetSwapChain()->GetTextureRef();
+		OutShaderResourceTexture = OutTargetableTexture;
+		return true;
+	}
+
 	for (uint32 SwapChainIter = 0; SwapChainIter < SteamVRSwapChainLength; ++SwapChainIter)
 	{
 #if PLATFORM_MAC
