Reset OpenXR swapchain when session is closed

diff --git a/Engine/Plugins/Runtime/OpenXR/Source/OpenXRHMD/Private/OpenXRHMD.cpp b/Engine/Plugins/Runtime/OpenXR/Source/OpenXRHMD/Private/OpenXRHMD.cpp
--- a/Engine/Plugins/Runtime/OpenXR/Source/OpenXRHMD/Private/OpenXRHMD.cpp
+++ b/Engine/Plugins/Runtime/OpenXR/Source/OpenXRHMD/Private/OpenXRHMD.cpp
@@ -1106,6 +1106,9 @@ void FOpenXRHMD::CloseSession()
 {
 	if (Session != XR_NULL_HANDLE)
 	{
+		Swapchain.Reset();
+		DepthSwapchain.Reset();
+
 		// Clear up device spaces
 		for (auto& DeviceSpace : DeviceSpaces)
 		{
@@ -1131,6 +1134,7 @@ void FOpenXRHMD::CloseSession()
 		bIsReady = false;
 		bIsRunning = false;
 		bRunRequested = false;
+		bNeedReAllocatedDepth = true;
 	}
 }
 
