Fix missing Vulkan path warning on Hololens builds.

diff --git a/Engine/Plugins/Runtime/OpenXR/Source/OpenXRHMD/OpenXRHMD.Build.cs b/Engine/Plugins/Runtime/OpenXR/Source/OpenXRHMD/OpenXRHMD.Build.cs
--- a/Engine/Plugins/Runtime/OpenXR/Source/OpenXRHMD/OpenXRHMD.Build.cs
+++ b/Engine/Plugins/Runtime/OpenXR/Source/OpenXRHMD/OpenXRHMD.Build.cs
@@ -25,7 +25,7 @@ namespace UnrealBuildTool.Rules
             {
                 PrivateIncludePaths.Add(EngineDir + "/Source/Runtime/VulkanRHI/Private/Windows");
             }
-            else
+            else if (Target.Platform != UnrealTargetPlatform.Hololens)
             {
                 PrivateIncludePaths.Add(EngineDir + "/Source/Runtime/VulkanRHI/Private/" + Target.Platform);
             }
