Fix capitalization error (HoloLens not Hololens).

diff --git a/Engine/Plugins/Runtime/OpenXR/Source/OpenXRHMD/OpenXRHMD.Build.cs b/Engine/Plugins/Runtime/OpenXR/Source/OpenXRHMD/OpenXRHMD.Build.cs
--- a/Engine/Plugins/Runtime/OpenXR/Source/OpenXRHMD/OpenXRHMD.Build.cs
+++ b/Engine/Plugins/Runtime/OpenXR/Source/OpenXRHMD/OpenXRHMD.Build.cs
@@ -25,7 +25,7 @@ namespace UnrealBuildTool.Rules
             {
                 PrivateIncludePaths.Add(EngineDir + "/Source/Runtime/VulkanRHI/Private/Windows");
             }
-            else if (Target.Platform != UnrealTargetPlatform.Hololens)
+            else if (Target.Platform != UnrealTargetPlatform.HoloLens)
             {
                 PrivateIncludePaths.Add(EngineDir + "/Source/Runtime/VulkanRHI/Private/" + Target.Platform);
             }
