The hand tracking extension is out of its preview phase

diff --git a/Engine/Plugins/Runtime/OpenXR/Source/OpenXRHMD/Private/OpenXRHMD.cpp b/Engine/Plugins/Runtime/OpenXR/Source/OpenXRHMD/Private/OpenXRHMD.cpp
--- a/Engine/Plugins/Runtime/OpenXR/Source/OpenXRHMD/Private/OpenXRHMD.cpp
+++ b/Engine/Plugins/Runtime/OpenXR/Source/OpenXRHMD/Private/OpenXRHMD.cpp
@@ -395,9 +395,9 @@ bool FOpenXRHMDPlugin::PreInit()
 	}
 
 	// Switch to constant once out of preview and there is a header available.
-	if (AvailableExtensions.Contains("XR_MSFT_hand_interaction_preview"))
+	if (AvailableExtensions.Contains("XR_MSFT_hand_interaction"))
 	{
-		EnabledExtensions.Add("XR_MSFT_hand_interaction_preview");
+		EnabledExtensions.Add("XR_MSFT_hand_interaction");
 	}
 
 	// Enable layers, if specified by CVar.
diff --git a/Engine/Plugins/Runtime/OpenXR/Source/OpenXRInput/Private/OpenXRInput.cpp b/Engine/Plugins/Runtime/OpenXR/Source/OpenXRInput/Private/OpenXRInput.cpp
--- a/Engine/Plugins/Runtime/OpenXR/Source/OpenXRInput/Private/OpenXRInput.cpp
+++ b/Engine/Plugins/Runtime/OpenXR/Source/OpenXRInput/Private/OpenXRInput.cpp
@@ -205,9 +205,9 @@ void FOpenXRInputPlugin::FOpenXRInput::BuildActions()
 	Profiles.Add("OculusTouch", FInteractionProfile(GetPath(Instance, "/interaction_profiles/oculus/touch_controller"), true));
 	Profiles.Add("ValveIndex", FInteractionProfile(GetPath(Instance, "/interaction_profiles/valve/index_controller"), true));
 
-	if (OpenXRHMD->IsExtensionEnabled("XR_MSFT_hand_interaction_preview"))
+	if (OpenXRHMD->IsExtensionEnabled("XR_MSFT_hand_interaction"))
 	{
-		Profiles.Add("MicrosoftHandInteraction", FInteractionProfile(GetPath(Instance, "/interaction_profiles/microsoft/hand_interaction_preview"), false));
+		Profiles.Add("MicrosoftHandInteraction", FInteractionProfile(GetPath(Instance, "/interaction_profiles/microsoft/hand_interaction"), false));
 	}
 
 	// Generate a list of the sub-action paths so we can query the left/right hand individually
