Remove references to vr.MobileMultiView.Direct CVar
Rename Multi-View back to Instanced Stereo until we add Multi-View support in

diff --git a/Engine/Plugins/Runtime/GoogleVR/GoogleVRHMD/Source/GoogleVRHMD/Private/GoogleVRHMD.cpp b/Engine/Plugins/Runtime/GoogleVR/GoogleVRHMD/Source/GoogleVRHMD/Private/GoogleVRHMD.cpp
--- a/Engine/Plugins/Runtime/GoogleVR/GoogleVRHMD/Source/GoogleVRHMD/Private/GoogleVRHMD.cpp
+++ b/Engine/Plugins/Runtime/GoogleVR/GoogleVRHMD/Source/GoogleVRHMD/Private/GoogleVRHMD.cpp
@@ -269,7 +269,7 @@ FGoogleVRHMD::FGoogleVRHMD(const FAutoRegister& AutoRegister)
 	, bUseOffscreenFramebuffers(false)
 	, bIsInDaydreamMode(false)
 	, bForceStopPresentScene(false)
-	, bIsMobileMultiViewDirect(false)
+	, bIsMobileMultiView(false)
 	, NeckModelScale(1.0f)
 	, BaseOrientation(FQuat::Identity)
 	, PixelDensity(1.0f)
@@ -415,10 +415,8 @@ FGoogleVRHMD::FGoogleVRHMD(const FAutoRegister& AutoRegister)
 		// Query for direct multi-view
 		GSupportsMobileMultiView = gvr_is_feature_supported(GVRAPI, GVR_FEATURE_MULTIVIEW) && bUseOffscreenFramebuffers;
 		const auto CVarMobileMultiView = IConsoleManager::Get().FindTConsoleVariableDataInt(TEXT("vr.MobileMultiView"));
-		const auto CVarMobileMultiViewDirect = IConsoleManager::Get().FindTConsoleVariableDataInt(TEXT("vr.MobileMultiView.Direct"));
 		const bool bIsMobileMultiViewEnabled = (CVarMobileMultiView && CVarMobileMultiView->GetValueOnAnyThread() != 0);
-		const bool bIsMobileMultiViewDirectEnabled = (CVarMobileMultiViewDirect && CVarMobileMultiViewDirect->GetValueOnAnyThread() != 0);
-		bIsMobileMultiViewDirect = GSupportsMobileMultiView && bIsMobileMultiViewEnabled && bIsMobileMultiViewDirectEnabled;
+		bIsMobileMultiView = GSupportsMobileMultiView && bIsMobileMultiViewEnabled;
 
 		if(bUseOffscreenFramebuffers)
 		{
@@ -1332,7 +1330,7 @@ void FGoogleVRHMD::PostRenderViewFamily_RenderThread(FRHICommandListImmediate& R
 #endif  // GOOGLEVRHMD_SUPPORTED_INSTANT_PREVIEW_PLATFORMS
 
 #if GOOGLEVRHMD_SUPPORTED_INSTANT_PREVIEW_PLATFORMS
-bool FGoogleVRHMD::GetCurrentReferencePose(FQuat& CurrentOrientation, FVector& CurrentPosition) const 
+bool FGoogleVRHMD::GetCurrentReferencePose(FQuat& CurrentOrientation, FVector& CurrentPosition) const
 {
 	FMatrix TransposeHeadPoseUnreal;
 	memcpy(&TransposeHeadPoseUnreal.M, CurrentReferencePose.pose.transform, 4 * 4 * sizeof(float));
@@ -1349,7 +1347,7 @@ bool FGoogleVRHMD::GetCurrentReferencePose(FQuat& CurrentOrientation, FVector& C
 	return true;
 }
 
-FVector FGoogleVRHMD::GetLocalEyePos(const instant_preview::EyeView& EyeView) 
+FVector FGoogleVRHMD::GetLocalEyePos(const instant_preview::EyeView& EyeView)
 {
 	const float* mat = EyeView.eye_pose.transform;
 	FMatrix PoseMatrix(
@@ -1360,7 +1358,7 @@ FVector FGoogleVRHMD::GetLocalEyePos(const instant_preview::EyeView& EyeView)
 	return PoseMatrix.TransformPosition(FVector(0, 0, 0));
 }
 
-void FGoogleVRHMD::PushVideoFrame(const FColor* VideoFrameBuffer, int width, int height, int stride, instant_preview::PixelFormat pixel_format, instant_preview::ReferencePose reference_pose) 
+void FGoogleVRHMD::PushVideoFrame(const FColor* VideoFrameBuffer, int width, int height, int stride, instant_preview::PixelFormat pixel_format, instant_preview::ReferencePose reference_pose)
 {
 	instant_preview::Session *session = ip_static_server_acquire_active_session(IpServerHandle);
 	if (session != NULL && width > 0 && height > 0)
diff --git a/Engine/Plugins/Runtime/GoogleVR/GoogleVRHMD/Source/GoogleVRHMD/Private/GoogleVRHMD.h b/Engine/Plugins/Runtime/GoogleVR/GoogleVRHMD/Source/GoogleVRHMD/Private/GoogleVRHMD.h
--- a/Engine/Plugins/Runtime/GoogleVR/GoogleVRHMD/Source/GoogleVRHMD/Private/GoogleVRHMD.h
+++ b/Engine/Plugins/Runtime/GoogleVR/GoogleVRHMD/Source/GoogleVRHMD/Private/GoogleVRHMD.h
@@ -251,10 +251,10 @@ public:
 	/** Check if the application is running in Daydream mode*/
 	bool IsInDaydreamMode() const;
 
-	/** Check if mobile multi-view direct is enabled */
-	bool IsMobileMultiViewDirect() const
+	/** Check if mobile multi-view is enabled */
+	bool IsMobileMultiView() const
 	{
-		return bIsMobileMultiViewDirect;
+		return bIsMobileMultiView;
 	}
 
 	void SetSPMEnable(bool bEnable) const;
@@ -374,7 +374,7 @@ private:
 	bool		bUseOffscreenFramebuffers;
 	bool		bIsInDaydreamMode;
 	bool		bForceStopPresentScene;
-	bool		bIsMobileMultiViewDirect;
+	bool		bIsMobileMultiView;
 	float		NeckModelScale;
 	FQuat		BaseOrientation;
 
@@ -521,7 +521,7 @@ public:
 	// Begin FXRRenderTargetManager Interface //
 	////////////////////////////////////////////
 
-	/** 
+	/**
 	 * Returns an instance of the currently active render bridge (aka. custom present)
 	 */
 	virtual FXRRenderBridge* GetActiveRenderBridge_GameThread(bool bUseSeparateRenderTarget) override;
@@ -601,7 +601,7 @@ public:
 	// Begin IXRTrackingSystem Virtual Interface //
 	///////////////////////////////////////////////
 
-	/** 
+	/**
 	 * GoogleVR currently does not support late update
 	 */
 	virtual bool DoesSupportLateUpdate() const override
diff --git a/Engine/Plugins/Runtime/Oculus/OculusVR/Source/OculusHMD/Private/OculusHMD.cpp b/Engine/Plugins/Runtime/Oculus/OculusVR/Source/OculusHMD/Private/OculusHMD.cpp
--- a/Engine/Plugins/Runtime/Oculus/OculusVR/Source/OculusHMD/Private/OculusHMD.cpp
+++ b/Engine/Plugins/Runtime/Oculus/OculusVR/Source/OculusHMD/Private/OculusHMD.cpp
@@ -2437,20 +2437,11 @@ namespace OculusHMD
 		ovrpLayout Layout = ovrpLayout_DoubleWide;
 #if PLATFORM_ANDROID
 		static const auto CVarMobileMultiView = IConsoleManager::Get().FindTConsoleVariableDataInt(TEXT("vr.MobileMultiView"));
-		static const auto CVarMobileMultiViewDirect = IConsoleManager::Get().FindTConsoleVariableDataInt(TEXT("vr.MobileMultiView.Direct"));
 		static const auto CVarMobileHDR = IConsoleManager::Get().FindTConsoleVariableDataInt(TEXT("r.MobileHDR"));
 		const bool bIsMobileMultiViewEnabled = (CVarMobileMultiView && CVarMobileMultiView->GetValueOnAnyThread() != 0);
-		const bool bIsMobileMultiViewDirectEnabled = (CVarMobileMultiViewDirect && CVarMobileMultiViewDirect->GetValueOnAnyThread() != 0);
-		const bool bIsUsingDirectMobileMultiView = GSupportsMobileMultiView && bIsMobileMultiViewEnabled && bIsMobileMultiViewDirectEnabled;
+		const bool bIsUsingMobileMultiView = GSupportsMobileMultiView && bIsMobileMultiViewEnabled;
 		const bool bMobileHDR = CVarMobileHDR && CVarMobileHDR->GetValueOnAnyThread() == 1;
 
-		if (bIsMobileMultiViewEnabled && !bIsMobileMultiViewDirectEnabled)
-		{
-			static bool bDisplayedMultiViewError = false;
-			UE_CLOG(!bDisplayedMultiViewError, LogHMD, Error, TEXT("\"Mobile Multiview Direct\" must always be enabled if \"Mobile Multiview\" is enabled on Oculus Mobile HMD devices."));
-			bDisplayedMultiViewError = true;
-		}
-
 		if (bMobileHDR)
 		{
 			static bool bDisplayedHDRError = false;
@@ -2458,7 +2449,7 @@ namespace OculusHMD
 			bDisplayedHDRError = true;
 		}
 
-		if (Settings->Flags.bDirectMultiview && bIsUsingDirectMobileMultiView)
+		if (bIsUsingMobileMultiView)
 		{
 			Layout = ovrpLayout_Array;
 			Settings->Flags.bIsUsingDirectMultiview = true;
diff --git a/Engine/Plugins/Runtime/OpenXR/Source/OpenXRHMD/Private/OpenXRHMD.cpp b/Engine/Plugins/Runtime/OpenXR/Source/OpenXRHMD/Private/OpenXRHMD.cpp
--- a/Engine/Plugins/Runtime/OpenXR/Source/OpenXRHMD/Private/OpenXRHMD.cpp
+++ b/Engine/Plugins/Runtime/OpenXR/Source/OpenXRHMD/Private/OpenXRHMD.cpp
@@ -405,7 +405,7 @@ bool FOpenXRHMDPlugin::PreInit()
 	//   1. Download and build the OpenXR SDK from [REDACTED_PROFILE_URL]/OpenXR-SDK-Source (follow instructions at [REDACTED_PROFILE_URL]/OpenXR-SDK-Source/blob/master/BUILDING.md)
 	//	 2. Add a registry key under HKEY_LOCAL_MACHINE\SOFTWARE\Khronos\OpenXR\1\ApiLayers\Explicit, containing the path to the manifest file
 	//      (e.g. C:\OpenXR-SDK-Source-master\build\win64\src\api_layers\XrApiLayer_core_validation.json) <-- this file is downloaded as part of the SDK source, above
-	//   3. Copy the DLL from the build target at, for example, C:\OpenXR-SDK-Source-master\build\win64\src\api_layers\XrApiLayer_core_validation.dll to 
+	//   3. Copy the DLL from the build target at, for example, C:\OpenXR-SDK-Source-master\build\win64\src\api_layers\XrApiLayer_core_validation.dll to
 	//      somewhere in your system path (e.g. c:\windows\system32); the OpenXR loader currently doesn't use the path the json file is in (this is a bug)
 
 	const bool bEnableOpenXRValidationLayer = CVarEnableOpenXRValidationLayer.GetValueOnAnyThread() != 0;
@@ -710,12 +710,8 @@ void FOpenXRHMD::AdjustViewRect(EStereoscopicPass StereoPass, int32& X, int32& Y
 	const uint32 ViewIndex = GetViewIndexForPass(StereoPass);
 
 	const XrViewConfigurationView& Config = Configs[ViewIndex];
-
-	static const auto CVarMobileMultiViewDirect = IConsoleManager::Get().FindTConsoleVariableDataInt(TEXT("vr.MobileMultiView.Direct"));
-	const bool bIsMobileMultiViewDirectEnabled = (CVarMobileMultiViewDirect && CVarMobileMultiViewDirect->GetValueOnAnyThread() != 0);
-
 	FIntPoint ViewRectMin(EForceInit::ForceInitToZero);
-	if (!bIsMobileMultiViewDirectEnabled)
+	if (!bIsMobileMultiViewEnabled)
 	{
 		for (uint32 i = 0; i < ViewIndex; ++i)
 		{
@@ -902,6 +898,11 @@ FOpenXRHMD::FOpenXRHMD(const FAutoRegister& AutoRegister, XrInstance InInstance,
 	bHiddenAreaMaskSupported = IsExtensionEnabled(XR_KHR_VISIBILITY_MASK_EXTENSION_NAME);
 #endif
 
+	static const auto CVarMobileMultiView = IConsoleManager::Get().FindTConsoleVariableDataInt(TEXT("vr.MobileMultiView"));
+	static const auto CVarMobileHDR = IConsoleManager::Get().FindTConsoleVariableDataInt(TEXT("r.MobileHDR"));
+	const bool bMobileHDR = (CVarMobileHDR && CVarMobileHDR->GetValueOnAnyThread() != 0);
+	bIsMobileMultiViewEnabled = (GSupportsMobileMultiView && !bMobileHDR) && (CVarMobileMultiView && CVarMobileMultiView->GetValueOnAnyThread() != 0);
+
 	{
 		// Enumerate the viewport configurations
 		uint32 ConfigurationCount;
@@ -1054,7 +1055,7 @@ bool FOpenXRHMD::BuildOcclusionMesh(XrVisibilityMaskTypeKHR Type, int View, FHMD
 
 	if (Type == XR_VISIBILITY_MASK_TYPE_HIDDEN_TRIANGLE_MESH_KHR)
 	{
-		// For the hidden-area mesh, bias z to the near plane so we occlude everything, and re-bias vertices to the -0.5 to 0.5 range in x and y. 
+		// For the hidden-area mesh, bias z to the near plane so we occlude everything, and re-bias vertices to the -0.5 to 0.5 range in x and y.
 		for (uint32 VertexIndex = 0; VertexIndex < VisibilityMask.vertexCountOutput; ++VertexIndex)
 		{
 			FFilterVertex& Vertex = Vertices[VertexIndex];
@@ -1069,7 +1070,7 @@ bool FOpenXRHMD::BuildOcclusionMesh(XrVisibilityMaskTypeKHR Type, int View, FHMD
 	}
 	else if (Type == XR_VISIBILITY_MASK_TYPE_VISIBLE_TRIANGLE_MESH_KHR)
 	{
-		// For the visible-area mesh, this will be consumed by the post-process pipeline, so set up coordinates in the space they expect 
+		// For the visible-area mesh, this will be consumed by the post-process pipeline, so set up coordinates in the space they expect
 		// (x and y range from 0-1, y reversed, z at the far plane).
 		for (uint32 VertexIndex = 0; VertexIndex < VisibilityMask.vertexCountOutput; ++VertexIndex)
 		{
@@ -1187,7 +1188,7 @@ bool FOpenXRHMD::OnStereoTeardown()
 		XrResult Result = xrRequestExitSession(Session);
 		if (Result == XR_ERROR_SESSION_NOT_RUNNING)
 		{
-			// Session was never running - most likely PIE without putting the headset on. 
+			// Session was never running - most likely PIE without putting the headset on.
 			CloseSession();
 		}
 		else
@@ -1263,7 +1264,7 @@ bool FOpenXRHMD::IsFocused() const
 
 bool FOpenXRHMD::StartSession()
 {
-	// If the session is already running, or is not yet ready, 
+	// If the session is already running, or is not yet ready,
 	if (!bIsReady || bIsRunning)
 	{
 		if (!bIsRunning)
@@ -1302,12 +1303,9 @@ bool FOpenXRHMD::AllocateRenderTargetTexture(uint32 Index, uint32 SizeX, uint32
 	// We need to ensure we can sample from the texture in CopyTexture
 	Flags |= TexCreate_ShaderResource;
 
-	static const auto CVarMobileMultiViewDirect = IConsoleManager::Get().FindTConsoleVariableDataInt(TEXT("vr.MobileMultiView.Direct"));
-	const bool bIsMobileMultiViewDirectEnabled = (CVarMobileMultiViewDirect && CVarMobileMultiViewDirect->GetValueOnAnyThread() != 0);
-
 	if (Swapchain == nullptr || Format != LastRequestedSwapchainFormat || Swapchain->GetTexture2D()->GetSizeX() != SizeX || Swapchain->GetTexture2D()->GetSizeY() != SizeY)
 	{
-		Swapchain = RenderBridge->CreateSwapchain(Session, Format, SizeX, SizeY, bIsMobileMultiViewDirectEnabled ? 2 : 1, NumMips, NumSamples, Flags, TargetableTextureFlags);
+		Swapchain = RenderBridge->CreateSwapchain(Session, Format, SizeX, SizeY, bIsMobileMultiViewEnabled ? 2 : 1, NumMips, NumSamples, Flags, TargetableTextureFlags);
 		if (!Swapchain)
 		{
 			return false;
@@ -1331,9 +1329,6 @@ bool FOpenXRHMD::AllocateDepthTexture(uint32 Index, uint32 SizeX, uint32 SizeY,
 		return false;
 	}
 
-	static const auto CVarMobileMultiView = IConsoleManager::Get().FindTConsoleVariableDataInt(TEXT("vr.MobileMultiView"));
-	const bool bIsMobileMultiViewEnabled = (CVarMobileMultiView && CVarMobileMultiView->GetValueOnAnyThread() != 0);
-
 	if (DepthSwapchain == nullptr || Format != LastRequestedDepthSwapchainFormat || DepthSwapchain->GetTexture2D()->GetSizeX() != SizeX || DepthSwapchain->GetTexture2D()->GetSizeY() != SizeY)
 	{
 		DepthSwapchain = RenderBridge->CreateSwapchain(Session, PF_DepthStencil, SizeX, SizeY, bIsMobileMultiViewEnabled ? 2 : 1, FMath::Max(NumMips, 1u), NumSamples, 0, TexCreate_DepthStencilTargetable);
@@ -1393,11 +1388,6 @@ void FOpenXRHMD::OnBeginRendering_RenderThread(FRHICommandListImmediate& RHICmdL
 	FIntPoint ViewRectMin(EForceInit::ForceInitToZero);
 	float NearZ = GNearClippingPlane / GetWorldToMetersScale();
 
-	static const auto CVarMobileMultiView = IConsoleManager::Get().FindTConsoleVariableDataInt(TEXT("vr.MobileMultiView"));
-	static const auto CVarMobileMultiViewDirect = IConsoleManager::Get().FindTConsoleVariableDataInt(TEXT("vr.MobileMultiView.Direct"));
-	const bool bIsMobileMultiViewEnabled = (CVarMobileMultiView && CVarMobileMultiView->GetValueOnAnyThread() != 0);
-	const bool bIsMobileMultiViewDirectEnabled = (CVarMobileMultiViewDirect && CVarMobileMultiViewDirect->GetValueOnAnyThread() != 0);
-
 	for (int32 ViewIndex = 0; ViewIndex < Views.Num(); ViewIndex++)
 	{
 		const XrView& View = Views[ViewIndex];
@@ -1416,7 +1406,7 @@ void FOpenXRHMD::OnBeginRendering_RenderThread(FRHICommandListImmediate& RHICmdL
 		Projection.fov = View.fov;
 		Projection.pose = ToXrPose(ViewTransform * BaseTransform, GetWorldToMetersScale());
 		Projection.subImage.swapchain = static_cast<FOpenXRSwapchain*>(GetSwapchain())->GetHandle();
-		Projection.subImage.imageArrayIndex = bIsMobileMultiViewDirectEnabled ? ViewIndex : 0;
+		Projection.subImage.imageArrayIndex = bIsMobileMultiViewEnabled ? ViewIndex : 0;
 		Projection.subImage.imageRect = {
 			{ ViewRect.Min.X, ViewRect.Min.Y },
 			{ ViewRect.Width(), ViewRect.Height() }
@@ -1437,7 +1427,7 @@ void FOpenXRHMD::OnBeginRendering_RenderThread(FRHICommandListImmediate& RHICmdL
 			Projection.next = &DepthLayer;
 		}
 
-		ViewRectMin.X += bIsMobileMultiViewDirectEnabled ? 0 : Config.recommendedImageRectWidth;
+		ViewRectMin.X += bIsMobileMultiViewEnabled ? 0 : Config.recommendedImageRectWidth;
 	}
 
 	// Give the RHI thread its own copy of the frame state and tracking space
@@ -1626,13 +1616,10 @@ FXRRenderBridge* FOpenXRHMD::GetActiveRenderBridge_GameThread(bool /* bUseSepara
 
 FIntPoint FOpenXRHMD::GetIdealRenderTargetSize() const
 {
-	static const auto CVarMobileMultiViewDirect = IConsoleManager::Get().FindTConsoleVariableDataInt(TEXT("vr.MobileMultiView.Direct"));
-	const bool bIsMobileMultiViewDirectEnabled = (CVarMobileMultiViewDirect && CVarMobileMultiViewDirect->GetValueOnAnyThread() != 0);
-
 	FIntPoint Size(EForceInit::ForceInitToZero);
 	for (XrViewConfigurationView Config : Configs)
 	{
-		Size.X = bIsMobileMultiViewDirectEnabled ? FMath::Max(Size.X, (int)Config.recommendedImageRectWidth)
+		Size.X = bIsMobileMultiViewEnabled ? FMath::Max(Size.X, (int)Config.recommendedImageRectWidth)
 			: Size.X + (int)Config.recommendedImageRectWidth;
 		Size.Y = FMath::Max(Size.Y, (int)Config.recommendedImageRectHeight);
 
diff --git a/Engine/Plugins/Runtime/OpenXR/Source/OpenXRHMD/Private/OpenXRHMD.h b/Engine/Plugins/Runtime/OpenXR/Source/OpenXRHMD/Private/OpenXRHMD.h
--- a/Engine/Plugins/Runtime/OpenXR/Source/OpenXRHMD/Private/OpenXRHMD.h
+++ b/Engine/Plugins/Runtime/OpenXR/Source/OpenXRHMD/Private/OpenXRHMD.h
@@ -206,7 +206,8 @@ private:
 	bool					bDepthExtensionSupported;
 	bool					bHiddenAreaMaskSupported;
 	bool					bNeedReAllocatedDepth;
-	
+	bool					bIsMobileMultiViewEnabled;
+
 	XrSessionState			CurrentSessionState;
 
 	FTransform				BaseTransform;
diff --git a/Engine/Plugins/Runtime/WindowsMixedReality/Source/WindowsMixedRealityHMD/Private/WindowsMixedRealityCustomPresent.h b/Engine/Plugins/Runtime/WindowsMixedReality/Source/WindowsMixedRealityHMD/Private/WindowsMixedRealityCustomPresent.h
--- a/Engine/Plugins/Runtime/WindowsMixedReality/Source/WindowsMixedRealityHMD/Private/WindowsMixedRealityCustomPresent.h
+++ b/Engine/Plugins/Runtime/WindowsMixedReality/Source/WindowsMixedRealityHMD/Private/WindowsMixedRealityCustomPresent.h
@@ -19,21 +19,18 @@ namespace WindowsMixedReality
 	public:
 		FWindowsMixedRealityCustomPresent(
 #if WITH_WINDOWS_MIXED_REALITY
-			MixedRealityInterop* _hmd, 
+			MixedRealityInterop* _hmd,
 #endif
-			ID3D11Device* device)
+			ID3D11Device* device,
+			bool bMultiView)
 			: FRHICustomPresent()
+			, bIsMultiViewEnabled(bMultiView)
 #if WITH_WINDOWS_MIXED_REALITY
 			, hmd(_hmd)
 #endif
 		{
 			// Get the D3D11 context.
 			device->GetImmediateContext(&D3D11Context);
-
-#if PLATFORM_HOLOLENS
-			static const auto CVarMobileMultiViewDirect = IConsoleManager::Get().FindTConsoleVariableDataInt(TEXT("vr.MobileMultiView.Direct"));
-			bIsMobileMultiViewDirectEnabled = (CVarMobileMultiViewDirect && CVarMobileMultiViewDirect->GetValueOnAnyThread() != 0);
-#endif
 		}
 
 		// Inherited via FRHICustomPresent
@@ -52,13 +49,13 @@ namespace WindowsMixedReality
 				return false;
 			}
 
-			if (!bIsMobileMultiViewDirectEnabled)
+			if (!bIsMultiViewEnabled)
 			{
 				hmd->CopyResources(D3D11Context, ViewportTexture);
 			}
 
 			InOutSyncInterval = 0;
-			
+
 			return hmd->Present();
 #else
 			return false;
@@ -96,6 +93,6 @@ namespace WindowsMixedReality
 
 		ID3D11DeviceContext* D3D11Context = nullptr;
 		ID3D11Texture2D* ViewportTexture = nullptr;
-		bool bIsMobileMultiViewDirectEnabled = false;
+		bool bIsMultiViewEnabled = false;
 	};
-}
\ No newline at end of file
+}
diff --git a/Engine/Plugins/Runtime/WindowsMixedReality/Source/WindowsMixedRealityHMD/Private/WindowsMixedRealityHMD.cpp b/Engine/Plugins/Runtime/WindowsMixedReality/Source/WindowsMixedRealityHMD/Private/WindowsMixedRealityHMD.cpp
--- a/Engine/Plugins/Runtime/WindowsMixedReality/Source/WindowsMixedRealityHMD/Private/WindowsMixedRealityHMD.cpp
+++ b/Engine/Plugins/Runtime/WindowsMixedReality/Source/WindowsMixedRealityHMD/Private/WindowsMixedRealityHMD.cpp
@@ -31,6 +31,10 @@
 
 #include "WindowsMixedRealityAvailability.h"
 
+#if SUPPORTS_WINDOWS_MIXED_REALITY_AR
+	#include "HoloLensModule.h"
+#endif
+
 #if WITH_INPUT_SIMULATION
 	#include "WindowsMixedRealityInputSimulationEngineSubsystem.h"
 #endif
@@ -106,7 +110,7 @@ namespace WindowsMixedReality
 		{
 #if WITH_WINDOWS_MIXED_REALITY
 			return HMD && HMD->IsAvailable();
-#endif 
+#endif
 			return false;
 		}
 
@@ -134,7 +138,7 @@ namespace WindowsMixedReality
 			if (HMD != nullptr)
 			{
 				HMD->SetLogCallback(&LogCallback);
-			}	
+			}
 #endif // WANTS_INTEROP_LOGGING
 
 #else // WITH_WINDOWS_MIXED_REALITY
@@ -422,7 +426,7 @@ namespace WindowsMixedReality
 
 			return true;
 		}
-		
+
 #if WITH_EDITOR
 #if WITH_WINDOWS_MIXED_REALITY
 		UpdateRemotingStatus();
@@ -547,7 +551,7 @@ namespace WindowsMixedReality
 		{
 			return;
 		}
-		
+
 		// Set Remoting status
 		HMDRemotingConnectionState state = HMD->GetConnectionState();
 
@@ -636,9 +640,7 @@ namespace WindowsMixedReality
 		}
 
 #if PLATFORM_HOLOLENS
-		static const auto CVarMobileMultiViewDirect = IConsoleManager::Get().FindTConsoleVariableDataInt(TEXT("vr.MobileMultiView.Direct"));
-		const bool bIsMobileMultiViewDirectEnabled = (CVarMobileMultiViewDirect && CVarMobileMultiViewDirect->GetValueOnRenderThread() != 0);
-		if (bIsMobileMultiViewDirectEnabled)
+		if (bIsMobileMultiViewEnabled)
 		{
 			FD3D11DynamicRHI* DynamicRHI = static_cast<FD3D11DynamicRHI*>(GDynamicRHI);
 			ID3D11Texture2D* Texture = HMD->GetBackBufferTexture();
@@ -664,8 +666,6 @@ namespace WindowsMixedReality
 			static FVector2D SrcNormRectMax(0.45f, 0.8f);
 
 #if PLATFORM_HOLOLENS
-			static const auto CVarMobileMultiView = IConsoleManager::Get().FindTConsoleVariableDataInt(TEXT("vr.MobileMultiView"));
-			const bool bIsMobileMultiViewEnabled = (CVarMobileMultiView && CVarMobileMultiView->GetValueOnRenderThread() != 0);
 			SrcNormRectMax.X = bIsMobileMultiViewEnabled ? 0.95f : 0.45f;
 #endif
 
@@ -791,8 +791,8 @@ namespace WindowsMixedReality
 			return false;
 		}
 		Frame& TheFrame = GetFrame();
-		FTransform relativeTransform = FTransform::Identity; 
-		
+		FTransform relativeTransform = FTransform::Identity;
+
 		switch (Eye)
 		{
 		case eSSP_LEFT_EYE:
@@ -822,7 +822,7 @@ namespace WindowsMixedReality
 	FMatrix FetchProjectionMatrix(EStereoscopicPass StereoPassType, WindowsMixedReality::MixedRealityInterop* HMD)
 	{
 		if (StereoPassType != eSSP_LEFT_EYE &&
-			StereoPassType != eSSP_RIGHT_EYE && 
+			StereoPassType != eSSP_RIGHT_EYE &&
 			StereoPassType != eSSP_THIRD_CAMERA_EYE)
 		{
 			return FMatrix::Identity;
@@ -1199,7 +1199,7 @@ namespace WindowsMixedReality
 			else
 			{
 				return FMatrix::Identity;
-			}		
+			}
 		}
 	}
 
@@ -1271,9 +1271,7 @@ namespace WindowsMixedReality
 #endif
 
 #if PLATFORM_HOLOLENS
-		static const auto CVarMobileMultiViewDirect = IConsoleManager::Get().FindTConsoleVariableDataInt(TEXT("vr.MobileMultiView.Direct"));
-		const bool bIsMobileMultiViewDirectEnabled = (CVarMobileMultiViewDirect && CVarMobileMultiViewDirect->GetValueOnAnyThread() != 0);
-		int Offset = bIsMobileMultiViewDirectEnabled ? 0 : Width;
+		int Offset = bIsMobileMultiViewEnabled ? 0 : Width;
 #else
 		int Offset = Width;
 #endif
@@ -1289,7 +1287,7 @@ namespace WindowsMixedReality
 		float eyeUVX = (float)Width / (float)Size.X;
 		float eyeUVY = (float)Height / (float)Size.Y;
 #if PLATFORM_HOLOLENS
-		float offUVX = bIsMobileMultiViewDirectEnabled ? 0.0f : eyeUVX;
+		float offUVX = bIsMobileMultiViewEnabled ? 0.0f : eyeUVX;
 #else
 		float offUVX = eyeUVX;
 #endif
@@ -1339,9 +1337,6 @@ namespace WindowsMixedReality
 		}
 
 #if PLATFORM_HOLOLENS
-		static const auto CVarMobileMultiView = IConsoleManager::Get().FindTConsoleVariableDataInt(TEXT("vr.MobileMultiView"));
-		const bool bIsMobileMultiViewEnabled = (CVarMobileMultiView && CVarMobileMultiView->GetValueOnRenderThread() != 0);
-
 		if (bIsMobileMultiViewEnabled)
 		{
 			ID3D11Texture2D* Texture = static_cast<ID3D11Texture2D*>(CurrentDepthBuffer->GetNativeResource());
@@ -1350,8 +1345,8 @@ namespace WindowsMixedReality
 		else
 #endif
 		{
-			// We keep refs to the depth texture in the hmd, so this function is non-const.	
-			// RenderTexture_RenderThread should perhaps be refactored or made non-const.  
+			// We keep refs to the depth texture in the hmd, so this function is non-const.
+			// RenderTexture_RenderThread should perhaps be refactored or made non-const.
 			// But to get this into a hotfix of 4.23 we shall simply cast.
 			FWindowsMixedRealityHMD* nonconstthis = const_cast<FWindowsMixedRealityHMD*>(this);
 			nonconstthis->CreateHMDDepthTexture(RHICmdList);
@@ -1381,9 +1376,7 @@ namespace WindowsMixedReality
 		FRHIResourceCreateInfo CreateInfo;
 
 #if PLATFORM_HOLOLENS
-		static const auto CVarMobileMultiViewDirect = IConsoleManager::Get().FindTConsoleVariableDataInt(TEXT("vr.MobileMultiView.Direct"));
-		const bool bIsMobileMultiViewDirectEnabled = (CVarMobileMultiViewDirect && CVarMobileMultiViewDirect->GetValueOnAnyThread() != 0);
-		if (bIsMobileMultiViewDirectEnabled)
+		if (bIsMobileMultiViewEnabled)
 		{
 			FTexture2DArrayRHIRef texture, resource;
 			RHICreateTargetableShaderResource2DArray(
@@ -1440,8 +1433,6 @@ namespace WindowsMixedReality
 		CreateInfo.ClearValueBinding = FClearValueBinding::DepthFar;
 
 #if PLATFORM_HOLOLENS
-		static const auto CVarMobileMultiView = IConsoleManager::Get().FindTConsoleVariableDataInt(TEXT("vr.MobileMultiView"));
-		const bool bIsMobileMultiViewEnabled = (CVarMobileMultiView && CVarMobileMultiView->GetValueOnAnyThread() != 0);
 		if (bIsMobileMultiViewEnabled)
 		{
 			FIntPoint size = GetIdealRenderTargetSize();
@@ -1659,7 +1650,7 @@ namespace WindowsMixedReality
 		{
 			// Create stereo depth texture
 			D3D11_TEXTURE2D_DESC tdesc;
-			
+
 			tdesc.Width = wmrRenderTargets[0].width;
 			tdesc.Height = wmrRenderTargets[0].height;
 			tdesc.MipLevels = 1;
@@ -1771,7 +1762,7 @@ namespace WindowsMixedReality
 		}
 
 		HMD->SetFocusPointForFrame(WMRUtility::ToMixedRealityVector(TrackingSpacePosition));
-#endif	
+#endif
 	}
 
 	bool FWindowsMixedRealityHMD::IsActiveThisFrame(class FViewport* InViewport) const
@@ -1791,6 +1782,11 @@ namespace WindowsMixedReality
 		static const FName RendererModuleName("Renderer");
 		RendererModule = FModuleManager::GetModulePtr<IRendererModule>(RendererModuleName);
 
+		static const auto CVarMobileMultiView = IConsoleManager::Get().FindTConsoleVariableDataInt(TEXT("vr.MobileMultiView"));
+		static const auto CVarMobileHDR = IConsoleManager::Get().FindTConsoleVariableDataInt(TEXT("r.MobileHDR"));
+		const bool bMobileHDR = (CVarMobileHDR && CVarMobileHDR->GetValueOnAnyThread() != 0);
+		bIsMobileMultiViewEnabled = (GSupportsMobileMultiView && !bMobileHDR) && (CVarMobileMultiView && CVarMobileMultiView->GetValueOnAnyThread() != 0);
+
 		HiddenAreaMesh.SetNum(2);
 		VisibleAreaMesh.SetNum(2);
 
@@ -1853,7 +1849,7 @@ namespace WindowsMixedReality
 			}
 		}
 #endif
-		
+
 		bNeedReallocateDepthTexture = false;
 
 		GetFrame().bPositionalTrackingUsed = false;
@@ -1874,7 +1870,7 @@ namespace WindowsMixedReality
 			VisibleAreaMesh[i].IndexBufferRHI = nullptr;
 			VisibleAreaMesh[i].VertexBufferRHI = nullptr;
 		}
-		
+
 #if WITH_EDITOR
 #if WITH_WINDOWS_MIXED_REALITY
 		prevState = HMDRemotingConnectionState::Undefined;
@@ -1994,7 +1990,7 @@ namespace WindowsMixedReality
 		//remove keys from "speech" namespace
 		EKeys::RemoveKeysWithCategory(FInputActionSpeechMapping::GetKeyCategory());
 #endif
-	}	   
+	}
 
 	bool WindowsMixedReality::FWindowsMixedRealityHMD::IsAvailable()
 	{
@@ -2011,7 +2007,7 @@ namespace WindowsMixedReality
 #if WITH_WINDOWS_MIXED_REALITY
 		if (mCustomPresent == nullptr)
 		{
-			mCustomPresent = new FWindowsMixedRealityCustomPresent(HMD, D3D11Device);
+			mCustomPresent = new FWindowsMixedRealityCustomPresent(HMD, D3D11Device, bIsMobileMultiViewEnabled);
 		}
 #endif
 	}
@@ -2164,7 +2160,7 @@ namespace WindowsMixedReality
 	{
 		HMD->SubmitHapticValue(hand, FMath::Clamp(value, 0.f, 1.f));
 	}
-	
+
 	bool FWindowsMixedRealityHMD::QueryCoordinateSystem(ABI::Windows::Perception::Spatial::ISpatialCoordinateSystem *& pCoordinateSystem, WindowsMixedReality::HMDTrackingOrigin& trackingOrigin)
 	{
 		return HMD->QueryCoordinateSystem(pCoordinateSystem, trackingOrigin);
@@ -2247,7 +2243,7 @@ namespace WindowsMixedReality
 #if WITH_WINDOWS_MIXED_REALITY
 		HMD->DisconnectFromDevice();
 		HMD->SetLogCallback(nullptr);
-		
+
 		// Close PIE if it's running:
 		UEditorEngine* EditorEngine = CastChecked<UEditorEngine>(GEngine);
 		FSceneViewport* PIEViewport = (FSceneViewport*)EditorEngine->GetPIEViewport();
diff --git a/Engine/Plugins/Runtime/WindowsMixedReality/Source/WindowsMixedRealityHMD/Private/WindowsMixedRealityHMD.h b/Engine/Plugins/Runtime/WindowsMixedReality/Source/WindowsMixedRealityHMD/Private/WindowsMixedRealityHMD.h
--- a/Engine/Plugins/Runtime/WindowsMixedReality/Source/WindowsMixedRealityHMD/Private/WindowsMixedRealityHMD.h
+++ b/Engine/Plugins/Runtime/WindowsMixedReality/Source/WindowsMixedRealityHMD/Private/WindowsMixedRealityHMD.h
@@ -121,10 +121,10 @@ namespace WindowsMixedReality
 		virtual FMatrix GetStereoProjectionMatrix(const enum EStereoscopicPass StereoPassType) const override;
 		virtual IStereoRenderTargetManager* GetRenderTargetManager() override { return this; }
 		virtual class IStereoLayers* GetStereoLayers() override;
-		
-		virtual int32 GetDesiredNumberOfViews(bool bStereoRequested) const override 
-		{ 
-			return (bStereoRequested) ? 3 : 1; 
+
+		virtual int32 GetDesiredNumberOfViews(bool bStereoRequested) const override
+		{
+			return (bStereoRequested) ? 3 : 1;
 		}
 
 		virtual EStereoscopicPass GetViewPassForIndex(bool bStereoRequested, uint32 ViewIndex) const override
@@ -207,6 +207,8 @@ namespace WindowsMixedReality
 		bool bRequestRestart = false;
 		bool bRequestShutdown = false;
 
+		bool bIsMobileMultiViewEnabled = false;
+
 		float ScreenScalePercentage = 1.0f;
 		float CachedWorldToMetersScale = 100.0f;
 
@@ -386,4 +388,3 @@ namespace WindowsMixedReality
 		void CreateSpectatorScreenController();
 	};
 }
-
diff --git a/Engine/Source/Runtime/Engine/Classes/Engine/RendererSettings.h b/Engine/Source/Runtime/Engine/Classes/Engine/RendererSettings.h
--- a/Engine/Source/Runtime/Engine/Classes/Engine/RendererSettings.h
+++ b/Engine/Source/Runtime/Engine/Classes/Engine/RendererSettings.h
@@ -598,7 +598,7 @@ class ENGINE_API URendererSettings : public UDeveloperSettings
 		uint32 bNvidiaAftermathEnabled : 1;
 
 	UPROPERTY(config, EditAnywhere, Category = VR, meta = (
-		ConsoleVariable = "vr.InstancedStereo", DisplayName = "Multi-View",
+		ConsoleVariable = "vr.InstancedStereo", DisplayName = "Instanced Stereo",
 		ToolTip = "Enable single-pass stereoscopic rendering through view instancing or draw call instancing.",
 		ConfigRestartRequired = true))
 	uint32 bMultiView : 1;
