Only process Hololens events if present when using OpenXR
Fix race condition in OpenXRHMD when locating pose

diff --git a/Engine/Plugins/Runtime/OpenXR/Source/OpenXRHMD/Private/OpenXRHMD.cpp b/Engine/Plugins/Runtime/OpenXR/Source/OpenXRHMD/Private/OpenXRHMD.cpp
--- a/Engine/Plugins/Runtime/OpenXR/Source/OpenXRHMD/Private/OpenXRHMD.cpp
+++ b/Engine/Plugins/Runtime/OpenXR/Source/OpenXRHMD/Private/OpenXRHMD.cpp
@@ -622,9 +622,24 @@ bool FOpenXRHMD::GetCurrentPose(int32 DeviceId, FQuat& CurrentOrientation, FVect
 		return false;
 	}
 
+	XrTime DisplayTime = 0;
+	if (IsInRenderingThread())
+	{
+		DisplayTime = FrameStateRHI.predictedDisplayTime;
+	}
+	else if (IsInGameThread())
+	{
+		DisplayTime = FrameState.predictedDisplayTime;
+	}
+	else
+	{
+		check(false)
+		return false;
+	}
+
 	XrSpaceLocation Location = {};
 	Location.type = XR_TYPE_SPACE_LOCATION;
-	XrResult Result = xrLocateSpace(DeviceSpaces[DeviceId].Space, GetTrackingSpace(), FrameState.predictedDisplayTime, &Location);
+	XrResult Result = xrLocateSpace(DeviceSpaces[DeviceId].Space, GetTrackingSpace(), DisplayTime, &Location);
 	if (!XR_ENSURE(Result))
 	{
 		return false;
@@ -1349,6 +1364,10 @@ bool FOpenXRHMD::AllocateDepthTexture(uint32 Index, uint32 SizeX, uint32 SizeY,
 
 void FOpenXRHMD::OnBeginRendering_RenderThread(FRHICommandListImmediate& RHICmdList, FSceneViewFamily& ViewFamily)
 {
+	FrameStateRHI = FrameState;
+
+	// xrBeginFrame is called _AFTER_ data the above data is cloned for RHI because xrBeginFrame will unblock xrWaitFrame in the game thread
+	// and that can lead to data races.
 	if (bIsReady)
 	{
 		XrFrameBeginInfo BeginInfo;
@@ -1377,8 +1396,6 @@ void FOpenXRHMD::OnBeginRendering_RenderThread(FRHICommandListImmediate& RHICmdL
 		}
 	}
 
-	FrameStateRHI = FrameState;
-
 	const FSceneView* MainView = ViewFamily.Views[0];
 	check(MainView);
 	BaseTransform = FTransform(MainView->BaseHmdOrientation, MainView->BaseHmdLocation);
diff --git a/Engine/Source/Runtime/Launch/Private/HoloLens/HoloLensLaunch.cpp b/Engine/Source/Runtime/Launch/Private/HoloLens/HoloLensLaunch.cpp
--- a/Engine/Source/Runtime/Launch/Private/HoloLens/HoloLensLaunch.cpp
+++ b/Engine/Source/Runtime/Launch/Private/HoloLens/HoloLensLaunch.cpp
@@ -56,6 +56,8 @@ using namespace Windows::System;
 using namespace Windows::UI::Core;
 using namespace Windows::Graphics::Holographic;
 
+static bool GSeparateCoreWindowVisible = false;
+
 DEFINE_LOG_CATEGORY_STATIC(LogLaunchHoloLens, Log, All);
 
 void appHoloLensEarlyInit();
@@ -582,7 +584,9 @@ void ViewProvider::ProcessEvents()
 
 		if (nullptr != Dispatcher)
 		{
-			if (bVisible)
+			// OpenXR uses a separate CoreWindow which will have visibility instead of the CoreWindow created by HoloLensLaunch and in this
+			// case ProcessOneAndAllPending cannot be used because it will block indefinitely once there are no more events to process.
+			if (bVisible || GSeparateCoreWindowVisible)
 			{
 				Dispatcher->ProcessEvents(CoreProcessEventsOption::ProcessAllIfPresent);
 			}
@@ -1232,6 +1236,9 @@ int32 GuardedMain( const TCHAR* CmdLine, HINSTANCE hInInstance, HINSTANCE hPrevI
 	// Game without wxWindows.
 	ErrorLevel = EngineInit();
 
+	static const FName OpenXRSystemName(TEXT("OpenXR"));
+	GSeparateCoreWindowVisible = GEngine->XRSystem.IsValid() && (GEngine->XRSystem->GetSystemName() == OpenXRSystemName);
+
 	uint64 GameThreadAffinity = FPlatformAffinity::GetMainGameMask();
 	FPlatformProcess::SetThreadAffinityMask((DWORD_PTR)GameThreadAffinity);
 
