Disable hidden area extension on OpenXR
Ignore invalid Valve Index controller key

diff --git a/Engine/Plugins/Runtime/OpenXR/Source/OpenXRHMD/Private/OpenXRHMD.cpp b/Engine/Plugins/Runtime/OpenXR/Source/OpenXRHMD/Private/OpenXRHMD.cpp
--- a/Engine/Plugins/Runtime/OpenXR/Source/OpenXRHMD/Private/OpenXRHMD.cpp
+++ b/Engine/Plugins/Runtime/OpenXR/Source/OpenXRHMD/Private/OpenXRHMD.cpp
@@ -915,7 +915,7 @@ FOpenXRHMD::FOpenXRHMD(const FAutoRegister& AutoRegister, XrInstance InInstance,
 	bIsMobileMultiViewEnabled = (GRHISupportsArrayIndexFromAnyShader && !bMobileHDR) && (CVarMobileMultiView && CVarMobileMultiView->GetValueOnAnyThread() != 0);
 #else
 	bDepthExtensionSupported = IsExtensionEnabled(XR_KHR_COMPOSITION_LAYER_DEPTH_EXTENSION_NAME);
-	bHiddenAreaMaskSupported = IsExtensionEnabled(XR_KHR_VISIBILITY_MASK_EXTENSION_NAME);
+	bHiddenAreaMaskSupported = false;
 	bIsMobileMultiViewEnabled = false;
 #endif
 
diff --git a/Engine/Plugins/Runtime/OpenXR/Source/OpenXRInput/Private/OpenXRInput.cpp b/Engine/Plugins/Runtime/OpenXR/Source/OpenXRInput/Private/OpenXRInput.cpp
--- a/Engine/Plugins/Runtime/OpenXR/Source/OpenXRInput/Private/OpenXRInput.cpp
+++ b/Engine/Plugins/Runtime/OpenXR/Source/OpenXRInput/Private/OpenXRInput.cpp
@@ -364,6 +364,12 @@ int32 FOpenXRInputPlugin::FOpenXRInput::SuggestBindings(XrInstance Instance, FOp
 			}
 			else if (Component == "click")
 			{
+				if (Identifier == "trackpad" && Tokens[0] == "ValveIndex")
+				{
+					// TODO: Remove click action for trackpad on Valve Index controller
+					continue;
+				}
+
 				if (Identifier != "trigger" && Identifier != "squeeze")
 				{
 					Path += "/click";
