Improve frame pipelining in the OpenXR plugin

diff --git a/Engine/Plugins/Runtime/OpenXR/Source/OpenXRHMD/Private/OpenXRHMD.cpp b/Engine/Plugins/Runtime/OpenXR/Source/OpenXRHMD/Private/OpenXRHMD.cpp
--- a/Engine/Plugins/Runtime/OpenXR/Source/OpenXRHMD/Private/OpenXRHMD.cpp
+++ b/Engine/Plugins/Runtime/OpenXR/Source/OpenXRHMD/Private/OpenXRHMD.cpp
@@ -617,45 +617,25 @@ bool FOpenXRHMD::GetCurrentPose(int32 DeviceId, FQuat& CurrentOrientation, FVect
 	CurrentOrientation = FQuat::Identity;
 	CurrentPosition = FVector::ZeroVector;
 
-	if (!DeviceSpaces.IsValidIndex(DeviceId) || !DeviceSpaces[DeviceId].Space || FrameState.predictedDisplayTime <= 0)
-	{
-		return false;
-	}
+	// This function is called from both the game and rendering thread and each thread maintains separate pose
+	// snapshots to prevent inconsistent poses (tearing) on the same frame.
+	const FPipelinedFrameState& PipelineState = GetPipelinedFrameStateForThread();
 
-	XrTime DisplayTime = 0;
-	if (IsInRenderingThread())
+	if (!PipelineState.DeviceLocations.IsValidIndex(DeviceId))
 	{
-		DisplayTime = FrameStateRHI.predictedDisplayTime;
-	}
-	else if (IsInGameThread())
-	{
-		DisplayTime = FrameState.predictedDisplayTime;
-	}
-	else
-	{
-		check(false)
 		return false;
 	}
 
-	XrSpaceLocation Location = {};
-	Location.type = XR_TYPE_SPACE_LOCATION;
-	XrResult Result = xrLocateSpace(DeviceSpaces[DeviceId].Space, GetTrackingSpace(), DisplayTime, &Location);
-	if (!XR_ENSURE(Result))
-	{
-		return false;
-	}
-
-	if (Location.locationFlags & XR_SPACE_LOCATION_ORIENTATION_VALID_BIT)
+	const XrSpaceLocation& Location = PipelineState.DeviceLocations[DeviceId];
+	if (Location.locationFlags & XR_SPACE_LOCATION_ORIENTATION_VALID_BIT &&
+		Location.locationFlags & XR_SPACE_LOCATION_POSITION_VALID_BIT)
 	{
 		CurrentOrientation = ToFQuat(Location.pose.orientation);
-	}
-
-	if (Location.locationFlags & XR_SPACE_LOCATION_POSITION_VALID_BIT)
-	{
 		CurrentPosition = ToFVector(Location.pose.position, GetWorldToMetersScale());
+		return true;
 	}
 
-	return true;
+	return false;
 }
 
 bool FOpenXRHMD::IsChromaAbCorrectionEnabled() const
@@ -776,7 +756,7 @@ uint32 FOpenXRHMD::DeviceGetLODViewIndex() const
 
 int32 FOpenXRHMD::GetDesiredNumberOfViews(bool bStereoRequested) const
 {
-	return bStereoRequested ? Views.Num() : 1; // FIXME: Monoscopic actually needs 2 views for quad vr
+	return bStereoRequested ? Configs.Num() : 1; // FIXME: Monoscopic actually needs 2 views for quad vr
 }
 
 bool FOpenXRHMD::GetRelativeEyePose(int32 InDeviceId, EStereoscopicPass InEye, FQuat& OutOrientation, FVector& OutPosition)
@@ -786,18 +766,26 @@ bool FOpenXRHMD::GetRelativeEyePose(int32 InDeviceId, EStereoscopicPass InEye, F
 		return false;
 	}
 
-	const uint32 ViewIndex = GetViewIndexForPass(InEye);
-	const XrView& View = Views[ViewIndex];
-	OutOrientation = ToFQuat(View.pose.orientation);
-	OutPosition = ToFVector(View.pose.position, GetWorldToMetersScale());
-	return true;
+	const FPipelinedFrameState& FrameState = GetPipelinedFrameStateForThread();
+
+	if (FrameState.ViewState.viewStateFlags & XR_VIEW_STATE_ORIENTATION_VALID_BIT &&
+		FrameState.ViewState.viewStateFlags & XR_VIEW_STATE_POSITION_VALID_BIT)
+	{
+		const uint32 ViewIndex = GetViewIndexForPass(InEye);
+		OutOrientation = ToFQuat(FrameState.Views[ViewIndex].pose.orientation);
+		OutPosition = ToFVector(FrameState.Views[ViewIndex].pose.position, GetWorldToMetersScale());
+		return true;
+	}
+
+	return false;
 }
 
 FMatrix FOpenXRHMD::GetStereoProjectionMatrix(const enum EStereoscopicPass StereoPassType) const
 {
 	const uint32 ViewIndex = GetViewIndexForPass(StereoPassType);
 
-	XrFovf Fov = Views[ViewIndex].fov;
+	const FPipelinedFrameState& FrameState = GetPipelinedFrameStateForThread();
+	XrFovf Fov = FrameState.Views[ViewIndex].fov;
 
 	float ZNear = GNearClippingPlane;
 
@@ -897,15 +885,6 @@ FOpenXRHMD::FOpenXRHMD(const FAutoRegister& AutoRegister, XrInstance InInstance,
 	, LastRequestedSwapchainFormat(0)
 	, LastRequestedDepthSwapchainFormat(0)
 {
-	FrameState.type = XR_TYPE_FRAME_STATE;
-	FrameState.next = nullptr;
-	FrameState.predictedDisplayPeriod = FrameState.predictedDisplayTime = 0;
-	FrameStateRHI = FrameState;
-
-	ViewState.type = XR_TYPE_VIEW_STATE;
-	ViewState.next = nullptr;
-	ViewState.viewStateFlags = 0;
-
 #if PLATFORM_HOLOLENS
 	bDepthExtensionSupported = bHiddenAreaMaskSupported = false;
 
@@ -960,8 +939,8 @@ FOpenXRHMD::FOpenXRHMD(const FAutoRegister& AutoRegister, XrInstance InInstance,
 	}
 
 	// Ensure the views have sane values before we locate them
-	Views.SetNum(Configs.Num());
-	for (XrView& View : Views)
+	PipelinedFrameStateGame.Views.SetNum(Configs.Num());
+	for (XrView& View : PipelinedFrameStateGame.Views)
 	{
 		View.type = XR_TYPE_VIEW;
 		View.next = nullptr;
@@ -1000,6 +979,9 @@ FOpenXRHMD::FOpenXRHMD(const FAutoRegister& AutoRegister, XrInstance InInstance,
 
 	// Add a device space for the HMD without an action handle and ensure it has the correct index
 	ensure(DeviceSpaces.Emplace(XR_NULL_HANDLE) == HMDDeviceId);
+
+	// Give the RHI frame state the same initial values.
+	PipelinedFrameStateRHI = PipelinedFrameStateGame;
 }
 
 FOpenXRHMD::~FOpenXRHMD()
@@ -1010,15 +992,58 @@ FOpenXRHMD::~FOpenXRHMD()
 	}
 }
 
+
+const FOpenXRHMD::FPipelinedFrameState& FOpenXRHMD::GetPipelinedFrameStateForThread() const
+{
+	if (IsInRenderingThread())
+	{
+		return PipelinedFrameStateRHI;
+	}
+	else
+	{
+		ensure(IsInGameThread());
+		return PipelinedFrameStateGame;
+	}
+}
+
+FOpenXRHMD::FPipelinedFrameState& FOpenXRHMD::GetPipelinedFrameStateForThread()
+{
+	if (IsInRenderingThread())
+	{
+		return PipelinedFrameStateRHI;
+	}
+	else
+	{
+		ensure(IsInGameThread());
+		return PipelinedFrameStateGame;
+	}
+}
+
+void FOpenXRHMD::UpdateDeviceLocations()
+{
+	FPipelinedFrameState& PipelineState = GetPipelinedFrameStateForThread();
+
+	PipelineState.DeviceLocations.SetNum(DeviceSpaces.Num());
+	for (int32 DeviceIndex = 0; DeviceIndex < PipelineState.DeviceLocations.Num(); DeviceIndex++)
+	{
+		const FDeviceSpace& DeviceSpace = DeviceSpaces[DeviceIndex];
+
+		XrSpaceLocation& DeviceLocation = PipelineState.DeviceLocations[DeviceIndex];
+		DeviceLocation.type = XR_TYPE_SPACE_LOCATION;
+		DeviceLocation.next = nullptr;
+		XR_ENSURE(xrLocateSpace(DeviceSpace.Space, PipelineState.TrackingSpace, PipelineState.FrameState.predictedDisplayTime, &DeviceLocation));
+	}
+}
+
 #if !PLATFORM_HOLOLENS
 void FOpenXRHMD::BuildOcclusionMeshes()
 {
-	HiddenAreaMeshes.SetNum(Views.Num());
-	VisibleAreaMeshes.SetNum(Views.Num());
+	HiddenAreaMeshes.SetNum(Configs.Num());
+	VisibleAreaMeshes.SetNum(Configs.Num());
 
 	bool bSucceeded = true;
 
-	for (int View = 0; View < Views.Num(); ++View)
+	for (int View = 0; View < Configs.Num(); ++View)
 	{
 		if (!BuildOcclusionMesh(XR_VISIBILITY_MASK_TYPE_VISIBLE_TRIANGLE_MESH_KHR, View, VisibleAreaMeshes[View]) ||
 			!BuildOcclusionMesh(XR_VISIBILITY_MASK_TYPE_HIDDEN_TRIANGLE_MESH_KHR, View, HiddenAreaMeshes[View]))
@@ -1179,7 +1204,7 @@ bool FOpenXRHMD::OnStereoStartup()
 	}
 
 	// Create action spaces for all devices
-	for (auto& DeviceSpace : DeviceSpaces)
+	for (FDeviceSpace& DeviceSpace : DeviceSpaces)
 	{
 		DeviceSpace.CreateSpace(Session);
 	}
@@ -1224,7 +1249,7 @@ void FOpenXRHMD::CloseSession()
 		DepthSwapchain.Reset();
 
 		// Clear up device spaces
-		for (auto& DeviceSpace : DeviceSpaces)
+		for (FDeviceSpace& DeviceSpace : DeviceSpaces)
 		{
 			DeviceSpace.DestroySpace();
 		}
@@ -1364,51 +1389,24 @@ bool FOpenXRHMD::AllocateDepthTexture(uint32 Index, uint32 SizeX, uint32 SizeY,
 
 void FOpenXRHMD::OnBeginRendering_RenderThread(FRHICommandListImmediate& RHICmdList, FSceneViewFamily& ViewFamily)
 {
-	FrameStateRHI = FrameState;
-
-	// xrBeginFrame is called _AFTER_ data the above data is cloned for RHI because xrBeginFrame will unblock xrWaitFrame in the game thread
-	// and that can lead to data races.
-	if (bIsReady)
-	{
-		XrFrameBeginInfo BeginInfo;
-		BeginInfo.type = XR_TYPE_FRAME_BEGIN_INFO;
-		BeginInfo.next = nullptr;
-		XrResult Result = xrBeginFrame(Session, &BeginInfo);
-		if (XR_SUCCEEDED(Result))
-		{
-			bIsRendering = true;
+	PipelinedFrameStateRHI = PipelinedFrameStateGame;
 
-			Swapchain->IncrementSwapChainIndex_RHIThread(FrameStateRHI.predictedDisplayPeriod);
-			if (bDepthExtensionSupported && !bNeedReAllocatedDepth)
-			{
-				ensure(DepthSwapchain != nullptr);
-				DepthSwapchain->IncrementSwapChainIndex_RHIThread(FrameStateRHI.predictedDisplayPeriod);
-			}
-		}
-		else
-		{
-			static bool bLoggedBeginFrameFailure = false;
-			if (!bLoggedBeginFrameFailure)
-			{
-				UE_LOG(LogHMD, Error, TEXT("Unexpected error on xrBeginFrame. Error code was %s."), OpenXRResultToString(Result));
-				bLoggedBeginFrameFailure = true;
-			}
-		}
-	}
+	ensure(IsInRenderingThread());
+	FPipelinedFrameState& PipelineState = GetPipelinedFrameStateForThread();
 
 	const FSceneView* MainView = ViewFamily.Views[0];
 	check(MainView);
-	BaseTransform = FTransform(MainView->BaseHmdOrientation, MainView->BaseHmdLocation);
+	const FTransform BaseTransform = FTransform(MainView->BaseHmdOrientation, MainView->BaseHmdLocation);
 
-	ViewsRHI.SetNum(Views.Num());
-	DepthLayersRHI.SetNum(Views.Num());
+	ProjectionViewsRHI.SetNum(PipelineState.Views.Num());
+	DepthLayersRHI.SetNum(PipelineState.Views.Num());
 
 	FIntPoint ViewRectMin(EForceInit::ForceInitToZero);
 	float NearZ = GNearClippingPlane / GetWorldToMetersScale();
 
-	for (int32 ViewIndex = 0; ViewIndex < Views.Num(); ViewIndex++)
+	for (int32 ViewIndex = 0; ViewIndex < PipelineState.Views.Num(); ViewIndex++)
 	{
-		const XrView& View = Views[ViewIndex];
+		const XrView& View = PipelineState.Views[ViewIndex];
 		const XrViewConfigurationView& Config = Configs[ViewIndex];
 
 		FIntRect ViewRect;
@@ -1416,7 +1414,7 @@ void FOpenXRHMD::OnBeginRendering_RenderThread(FRHICommandListImmediate& RHICmdL
 		ViewRect.Max = ViewRect.Min + FIntPoint(Config.recommendedImageRectWidth, Config.recommendedImageRectHeight);
 		FTransform ViewTransform = ToFTransform(View.pose, GetWorldToMetersScale());
 
-		XrCompositionLayerProjectionView& Projection = ViewsRHI[ViewIndex];
+		XrCompositionLayerProjectionView& Projection = ProjectionViewsRHI[ViewIndex];
 		XrCompositionLayerDepthInfoKHR& DepthLayer = DepthLayersRHI[ViewIndex];
 
 		Projection.type = XR_TYPE_COMPOSITION_LAYER_PROJECTION_VIEW;
@@ -1436,7 +1434,7 @@ void FOpenXRHMD::OnBeginRendering_RenderThread(FRHICommandListImmediate& RHICmdL
 			DepthLayer.next = nullptr;
 			DepthLayer.subImage.swapchain = static_cast<FOpenXRSwapchain*>(GetDepthSwapchain())->GetHandle();
 			DepthLayer.subImage.imageArrayIndex = bIsMobileMultiViewEnabled ? ViewIndex : 0;
-			DepthLayer.subImage.imageRect = bIsMobileMultiViewEnabled ? ViewsRHI[0].subImage.imageRect : Projection.subImage.imageRect;
+			DepthLayer.subImage.imageRect = bIsMobileMultiViewEnabled ? ProjectionViewsRHI[0].subImage.imageRect : Projection.subImage.imageRect;
 			DepthLayer.minDepth = 0.0f;
 			DepthLayer.maxDepth = 1.0f;
 			DepthLayer.nearZ = FLT_MAX;
@@ -1448,18 +1446,53 @@ void FOpenXRHMD::OnBeginRendering_RenderThread(FRHICommandListImmediate& RHICmdL
 		ViewRectMin.X += bIsMobileMultiViewEnabled ? 0 : Config.recommendedImageRectWidth;
 	}
 
-	// Give the RHI thread its own copy of the frame state and tracking space
-	TrackingSpaceRHI = GetTrackingSpace();
+	// xrBeginFrame is called _AFTER_ data the above data is cloned for RHI because xrBeginFrame will unblock xrWaitFrame in the game thread
+	// and that could otherwise lead to data races.
+	if (bIsReady)
+	{
+		XrFrameBeginInfo BeginInfo;
+		BeginInfo.type = XR_TYPE_FRAME_BEGIN_INFO;
+		BeginInfo.next = nullptr;
+		XrResult Result = xrBeginFrame(Session, &BeginInfo);
+		if (XR_SUCCEEDED(Result))
+		{
+			bIsRendering = true;
+
+			Swapchain->IncrementSwapChainIndex_RHIThread(PipelineState.FrameState.predictedDisplayPeriod);
+			if (bDepthExtensionSupported && !bNeedReAllocatedDepth)
+			{
+				ensure(DepthSwapchain != nullptr);
+				DepthSwapchain->IncrementSwapChainIndex_RHIThread(PipelineState.FrameState.predictedDisplayPeriod);
+			}
+		}
+		else
+		{
+			static bool bLoggedBeginFrameFailure = false;
+			if (!bLoggedBeginFrameFailure)
+			{
+				UE_LOG(LogHMD, Error, TEXT("Unexpected error on xrBeginFrame. Error code was %s."), OpenXRResultToString(Result));
+				bLoggedBeginFrameFailure = true;
+			}
+		}
+	}
+
+	// Snapshot new poses for late update.
+	UpdateDeviceLocations();
 }
 
 void FOpenXRHMD::OnLateUpdateApplied_RenderThread(const FTransform& NewRelativeTransform)
 {
 	FHeadMountedDisplayBase::OnLateUpdateApplied_RenderThread(NewRelativeTransform);
 
-	for (int32 ViewIndex = 0; ViewIndex < Views.Num(); ViewIndex++)
+	ensure(IsInRenderingThread());
+	const FPipelinedFrameState& PipelineState = GetPipelinedFrameStateForThread();
+
+	for (int32 ViewIndex = 0; ViewIndex < PipelineState.Views.Num(); ViewIndex++)
 	{
-		XrCompositionLayerProjectionView& Projection = ViewsRHI[ViewIndex];
-		FTransform ViewTransform = ToFTransform(Projection.pose, GetWorldToMetersScale()) * BaseTransform.Inverse();
+		const XrView& View = PipelineState.Views[ViewIndex];
+		XrCompositionLayerProjectionView& Projection = ProjectionViewsRHI[ViewIndex];
+
+		FTransform ViewTransform = ToFTransform(View.pose, GetWorldToMetersScale());
 		Projection.pose = ToXrPose(ViewTransform * NewRelativeTransform, GetWorldToMetersScale());
 	}
 }
@@ -1472,27 +1505,34 @@ void FOpenXRHMD::OnBeginRendering_GameThread()
 		return;
 	}
 
+	ensure(IsInGameThread());
+
 	XrFrameWaitInfo WaitInfo;
 	WaitInfo.type = XR_TYPE_FRAME_WAIT_INFO;
 	WaitInfo.next = nullptr;
 
-	// Pass in a clean FrameState (some implementations may fail if this is just uninitialized).
-	FrameState.type = XR_TYPE_FRAME_STATE;
-	FrameState.next = nullptr;
-	FrameState.predictedDisplayPeriod = 0;
-
+	XrFrameState FrameState{XR_TYPE_FRAME_STATE};
 	XR_ENSURE(xrWaitFrame(Session, &WaitInfo, &FrameState));
 
+	// The pipeline state on the game thread can only be safely modified after xrWaitFrame which will be unblocked by
+	// the runtime when xrBeginFrame is called. The rendering thread will clone the game pipeline state before calling
+	// xrBeginFrame so the game pipeline state can safely be modified after xrWaitFrame returns.
+	FPipelinedFrameState& PipelineState = GetPipelinedFrameStateForThread();
+	PipelineState.FrameState = FrameState;
+	PipelineState.TrackingSpace = GetTrackingSpace();
+
 	uint32_t ViewCount = 0;
 	XrViewLocateInfo ViewInfo;
 	ViewInfo.type = XR_TYPE_VIEW_LOCATE_INFO;
 	ViewInfo.next = nullptr;
 	ViewInfo.viewConfigurationType = SelectedViewConfigurationType;
 	ViewInfo.space = DeviceSpaces[HMDDeviceId].Space;
-	ViewInfo.displayTime = FrameState.predictedDisplayTime;
-	XR_ENSURE(xrLocateViews(Session, &ViewInfo, &ViewState, 0, &ViewCount, nullptr));
-	Views.SetNum(ViewCount);
-	XR_ENSURE(xrLocateViews(Session, &ViewInfo, &ViewState, Views.Num(), &ViewCount, Views.GetData()));
+	ViewInfo.displayTime = PipelineState.FrameState.predictedDisplayTime;
+	XR_ENSURE(xrLocateViews(Session, &ViewInfo, &PipelineState.ViewState, 0, &ViewCount, nullptr));
+	PipelineState.Views.SetNum(ViewCount);
+	XR_ENSURE(xrLocateViews(Session, &ViewInfo, &PipelineState.ViewState, PipelineState.Views.Num(), &ViewCount, PipelineState.Views.GetData()));
+
+	UpdateDeviceLocations();
 }
 
 bool FOpenXRHMD::ReadNextEvent(XrEventDataBuffer* buffer)
@@ -1594,12 +1634,15 @@ bool FOpenXRHMD::OnStartGameFrame(FWorldContext& WorldContext)
 
 void FOpenXRHMD::FinishRendering()
 {
+	ensure(IsInRenderingThread());
+	const FPipelinedFrameState& PipelineState = GetPipelinedFrameStateForThread();
+
 	XrCompositionLayerProjection Layer = {};
 	Layer.type = XR_TYPE_COMPOSITION_LAYER_PROJECTION;
 	Layer.next = nullptr;
-	Layer.space = TrackingSpaceRHI;
-	Layer.viewCount = ViewsRHI.Num();
-	Layer.views = ViewsRHI.GetData();
+	Layer.space = PipelineState.TrackingSpace;
+	Layer.viewCount = ProjectionViewsRHI.Num();
+	Layer.views = ProjectionViewsRHI.GetData();
 
 	if (bIsRendering)
 	{
@@ -1614,7 +1657,7 @@ void FOpenXRHMD::FinishRendering()
 		XrCompositionLayerBaseHeader* Headers[1] = { reinterpret_cast<XrCompositionLayerBaseHeader*>(&Layer) };
 		EndInfo.type = XR_TYPE_FRAME_END_INFO;
 		EndInfo.next = nullptr;
-		EndInfo.displayTime = FrameStateRHI.predictedDisplayTime;
+		EndInfo.displayTime = PipelineState.FrameState.predictedDisplayTime;
 		EndInfo.environmentBlendMode = SelectedEnvironmentBlendMode;
 		EndInfo.layerCount = 1;
 		EndInfo.layers = reinterpret_cast<XrCompositionLayerBaseHeader**>(Headers);
diff --git a/Engine/Plugins/Runtime/OpenXR/Source/OpenXRHMD/Private/OpenXRHMD.h b/Engine/Plugins/Runtime/OpenXR/Source/OpenXRHMD/Private/OpenXRHMD.h
--- a/Engine/Plugins/Runtime/OpenXR/Source/OpenXRHMD/Private/OpenXRHMD.h
+++ b/Engine/Plugins/Runtime/OpenXR/Source/OpenXRHMD/Private/OpenXRHMD.h
@@ -39,6 +39,17 @@ public:
 		XrSpace Space;
 	};
 
+	// The game and render threads each have a separate copy of this structure so that they don't stomp on each other or cause tearing
+	// when the game thread progresses to the next frame while the render thread is still working on the previous frame.
+	struct FPipelinedFrameState
+	{
+		XrFrameState FrameState{XR_TYPE_FRAME_STATE};
+		XrViewState ViewState{XR_TYPE_VIEW_STATE};
+		TArray<XrView> Views;
+		TArray<XrSpaceLocation> DeviceLocations;
+		XrSpace TrackingSpace;
+	};
+
 	class FVulkanExtensions : public IHeadMountedDisplayVulkanExtensions
 	{
 	public:
@@ -113,6 +124,10 @@ protected:
 	void BuildOcclusionMeshes();
 	bool BuildOcclusionMesh(XrVisibilityMaskTypeKHR Type, int View, FHMDViewMesh& Mesh);
 
+	const FPipelinedFrameState& GetPipelinedFrameStateForThread() const;
+	FPipelinedFrameState& GetPipelinedFrameStateForThread();
+	void UpdateDeviceLocations();
+
 public:
 	/** IHeadMountedDisplay interface */
 	virtual bool IsHMDConnected() override { return true; }
@@ -210,25 +225,21 @@ private:
 
 	XrSessionState			CurrentSessionState;
 
-	FTransform				BaseTransform;
 	TArray<const char*>		EnabledExtensions;
 	XrInstance				Instance;
 	XrSystemId				System;
 	XrSession				Session;
 	XrSpace					LocalSpace;
 	XrSpace					StageSpace;
-	XrSpace					TrackingSpaceRHI;
 	XrReferenceSpaceType	TrackingSpaceType;
 	XrViewConfigurationType SelectedViewConfigurationType;
 	XrEnvironmentBlendMode  SelectedEnvironmentBlendMode;
 
-	XrFrameState			FrameState;
-	XrFrameState			FrameStateRHI;
-	XrViewState				ViewState;
+	FPipelinedFrameState	PipelinedFrameStateGame;
+	FPipelinedFrameState	PipelinedFrameStateRHI;
 
 	TArray<XrViewConfigurationView> Configs;
-	TArray<XrView>			Views;
-	TArray<XrCompositionLayerProjectionView> ViewsRHI;
+	TArray<XrCompositionLayerProjectionView> ProjectionViewsRHI;
 	TArray<XrCompositionLayerDepthInfoKHR> DepthLayersRHI;
 
 	TArray<FDeviceSpace>	DeviceSpaces;
