Build fixes.

#rnx

diff --git a/Engine/Plugins/Runtime/OpenXR/Source/OpenXRHMD/Private/OpenXRHMD.cpp b/Engine/Plugins/Runtime/OpenXR/Source/OpenXRHMD/Private/OpenXRHMD.cpp
--- a/Engine/Plugins/Runtime/OpenXR/Source/OpenXRHMD/Private/OpenXRHMD.cpp
+++ b/Engine/Plugins/Runtime/OpenXR/Source/OpenXRHMD/Private/OpenXRHMD.cpp
@@ -28,6 +28,7 @@
 #include "BuildSettings.h"
 #include "ARSystem.h"
 #include "IHandTracker.h"
+#include "ScenePrivate.h"
 
 #if PLATFORM_ANDROID
 #include <android_native_app_glue.h>
