Merge OpenXR plugin from //UE4/Release-4.27

diff --git a/Engine/Plugins/Runtime/Oculus/OculusVR/Source/OculusHMD/Private/OculusHMD.cpp b/Engine/Plugins/Runtime/Oculus/OculusVR/Source/OculusHMD/Private/OculusHMD.cpp
--- a/Engine/Plugins/Runtime/Oculus/OculusVR/Source/OculusHMD/Private/OculusHMD.cpp
+++ b/Engine/Plugins/Runtime/Oculus/OculusVR/Source/OculusHMD/Private/OculusHMD.cpp
@@ -1277,7 +1277,7 @@ namespace OculusHMD
 		}
 	}
 
-	void FOculusHMD::SetFinalViewRect(const enum EStereoscopicPass StereoPass, const FIntRect& FinalViewRect)
+	void FOculusHMD::SetFinalViewRect(FRHICommandListImmediate& RHICmdList, const enum EStereoscopicPass StereoPass, const FIntRect& FinalViewRect)
 	{
 		CheckInRenderThread();
 
diff --git a/Engine/Plugins/Runtime/Oculus/OculusVR/Source/OculusHMD/Private/OculusHMD.h b/Engine/Plugins/Runtime/Oculus/OculusVR/Source/OculusHMD/Private/OculusHMD.h
--- a/Engine/Plugins/Runtime/Oculus/OculusVR/Source/OculusHMD/Private/OculusHMD.h
+++ b/Engine/Plugins/Runtime/Oculus/OculusVR/Source/OculusHMD/Private/OculusHMD.h
@@ -161,7 +161,7 @@ public:
 	virtual bool IsStereoEnabledOnNextFrame() const override;
 	virtual bool EnableStereo(bool stereo = true) override;
 	virtual void AdjustViewRect(enum EStereoscopicPass StereoPass, int32& X, int32& Y, uint32& SizeX, uint32& SizeY) const  override;
-	virtual void SetFinalViewRect(const enum EStereoscopicPass StereoPass, const FIntRect& FinalViewRect) override;
+	virtual void SetFinalViewRect(FRHICommandListImmediate& RHICmdList, const enum EStereoscopicPass StereoPass, const FIntRect& FinalViewRect) override;
 	//virtual FVector2D GetTextSafeRegionBounds() const override;
 	virtual void CalculateStereoViewOffset(const enum EStereoscopicPass StereoPassType, FRotator& ViewRotation, const float WorldToMeters, FVector& ViewLocation) override;
 	virtual FMatrix GetStereoProjectionMatrix(const enum EStereoscopicPass StereoPassType) const override;
diff --git a/Engine/Plugins/Runtime/OpenXR/Source/OpenXRHMD/Private/OpenXRHMD.cpp b/Engine/Plugins/Runtime/OpenXR/Source/OpenXRHMD/Private/OpenXRHMD.cpp
--- a/Engine/Plugins/Runtime/OpenXR/Source/OpenXRHMD/Private/OpenXRHMD.cpp
+++ b/Engine/Plugins/Runtime/OpenXR/Source/OpenXRHMD/Private/OpenXRHMD.cpp
@@ -29,6 +29,7 @@
 #include "BuildSettings.h"
 #include "ARSystem.h"
 #include "IHandTracker.h"
+#include "PixelShaderUtils.h"
 #include "ScenePrivate.h"
 
 #if PLATFORM_ANDROID
@@ -860,12 +861,34 @@ bool FOpenXRHMD::FVulkanExtensions::GetVulkanDeviceExtensionsRequired(VkPhysical
 
 void FOpenXRHMD::GetMotionControllerData(UObject* WorldContext, const EControllerHand Hand, FXRMotionControllerData& MotionControllerData)
 {
-	MotionControllerData.DeviceName = GetSystemName();
+	MotionControllerData.DeviceName = NAME_None;
 	MotionControllerData.ApplicationInstanceID = FApp::GetInstanceId();
 	MotionControllerData.DeviceVisualType = EXRVisualType::Controller;
 	MotionControllerData.TrackingStatus = ETrackingStatus::NotTracked;
 	MotionControllerData.HandIndex = Hand;
 
+	TArray<int32> Devices;
+	if (EnumerateTrackedDevices(Devices, EXRTrackedDeviceType::Controller) && Devices.IsValidIndex((int32)Hand))
+	{
+		FReadScopeLock SessionLock(SessionHandleMutex);
+		if (Session)
+		{
+			XrInteractionProfileState Profile;
+			Profile.type = XR_TYPE_INTERACTION_PROFILE_STATE;
+			Profile.next = nullptr;
+			if (XR_SUCCEEDED(xrGetCurrentInteractionProfile(Session, GetTrackedDevicePath(Devices[(int32)Hand]), &Profile)) &&
+				Profile.interactionProfile != XR_NULL_PATH)
+			{
+				TArray<char> Path;
+				uint32 PathCount = 0;
+				XR_ENSURE(xrPathToString(Instance, Profile.interactionProfile, 0, &PathCount, nullptr));
+				Path.SetNum(PathCount);
+				XR_ENSURE(xrPathToString(Instance, Profile.interactionProfile, PathCount, &PathCount, Path.GetData()));
+				MotionControllerData.DeviceName = Path.GetData();
+			}
+		}
+	}
+
 	FName HandTrackerName("OpenXRHandTracking");
 	TArray<IHandTracker*> HandTrackers = IModularFeatures::Get().GetModularFeatureImplementations<IHandTracker>(IHandTracker::GetModularFeatureName());
 	IHandTracker* HandTracker = nullptr;
@@ -955,6 +978,7 @@ FVector2D FOpenXRHMD::GetPlayAreaBounds(EHMDTrackingOrigin::Type Origin) const
 	}
 	XrExtent2Df Bounds;
 	XR_ENSURE(xrGetReferenceSpaceBoundsRect(Session, Space, &Bounds));
+	Swap(Bounds.width, Bounds.height); // Convert to UE coordinate system
 	return ToFVector2D(Bounds, WorldToMetersScale);
 }
 
@@ -1017,6 +1041,8 @@ bool FOpenXRHMD::EnumerateTrackedDevices(TArray<int32>& OutDevices, EXRTrackedDe
 	}
 	if (Type == EXRTrackedDeviceType::Any || Type == EXRTrackedDeviceType::Controller)
 	{
+		FReadScopeLock DeviceLock(DeviceMutex);
+
 		// Skip the HMD, we already added it to the list
 		for (int32 i = 1; i < DeviceSpaces.Num(); i++)
 		{
@@ -1089,6 +1115,7 @@ bool FOpenXRHMD::GetPoseForTime(int32 DeviceId, FTimespan Timespan, FQuat& Orien
 {
 	FPipelinedFrameState& PipelineState = GetPipelinedFrameStateForThread();
 
+	FReadScopeLock DeviceLock(DeviceMutex);
 	if (!DeviceSpaces.IsValidIndex(DeviceId))
 	{
 		return false;
@@ -1237,7 +1264,7 @@ void FOpenXRHMD::AdjustViewRect(EStereoscopicPass StereoPass, int32& X, int32& Y
 	SizeY = Config.recommendedImageRectHeight;
 }
 
-void FOpenXRHMD::SetFinalViewRect(const enum EStereoscopicPass StereoPass, const FIntRect& FinalViewRect)
+void FOpenXRHMD::SetFinalViewRect(FRHICommandListImmediate& RHICmdList, const enum EStereoscopicPass StereoPass, const FIntRect& FinalViewRect)
 {
 	if (StereoPass == eSSP_FULL)
 	{
@@ -1303,6 +1330,11 @@ void FOpenXRHMD::SetFinalViewRect(const enum EStereoscopicPass StereoPass, const
 	{
 		Projection.next = Module->OnBeginProjectionView(Session, 0, ViewIndex, Projection.next);
 	}
+
+	RHICmdList.EnqueueLambda([this, LayerState = PipelinedLayerStateRendering](FRHICommandListImmediate& InRHICmdList)
+	{
+		PipelinedLayerStateRHI = LayerState;
+	});
 }
 
 EStereoscopicPass FOpenXRHMD::GetViewPassForIndex(bool bStereoRequested, uint32 ViewIndex) const
@@ -1679,7 +1711,7 @@ void FOpenXRHMD::UpdateDeviceLocations(bool bUpdateOpenXRExtensionPlugins)
 	FPipelinedFrameState& PipelineState = GetPipelinedFrameStateForThread();
 
 	// Only update the device locations if the frame state has been predicted
-	if (bIsSynchronized && PipelineState.FrameState.predictedDisplayTime > 0)
+	if (PipelineState.FrameState.predictedDisplayTime > 0)
 	{
 		FReadScopeLock Lock(DeviceMutex);
 		PipelineState.DeviceLocations.SetNum(DeviceSpaces.Num());
@@ -1691,7 +1723,17 @@ void FOpenXRHMD::UpdateDeviceLocations(bool bUpdateOpenXRExtensionPlugins)
 				XrSpaceLocation& DeviceLocation = PipelineState.DeviceLocations[DeviceIndex];
 				DeviceLocation.type = XR_TYPE_SPACE_LOCATION;
 				DeviceLocation.next = nullptr;
-				XR_ENSURE(xrLocateSpace(DeviceSpace.Space, PipelineState.TrackingSpace, PipelineState.FrameState.predictedDisplayTime, &DeviceLocation));
+
+				XrResult Result = xrLocateSpace(DeviceSpace.Space, PipelineState.TrackingSpace, PipelineState.FrameState.predictedDisplayTime, &DeviceLocation);
+				if (Result == XR_ERROR_TIME_INVALID)
+				{
+					// The display time is no longer valid so set the location as invalid as well
+					PipelineState.DeviceLocations[DeviceIndex].locationFlags = 0;
+				}
+				else
+				{
+					XR_ENSURE(Result);
+				}
 			}
 			else
 			{
@@ -1756,6 +1798,7 @@ void FOpenXRHMD::EnumerateViews(FPipelinedFrameState& PipelineState)
 	{
 		LocateViews(PipelineState, true);
 
+		FReadScopeLock DeviceLock(DeviceMutex);
 		for (IOpenXRExtensionPlugin* Module : ExtensionPlugins)
 		{
 			if (PipelineState.PluginViews.Contains(Module))
@@ -1805,19 +1848,18 @@ void FOpenXRHMD::BuildOcclusionMeshes()
 	HiddenAreaMeshes.SetNum(ViewCount);
 	VisibleAreaMeshes.SetNum(ViewCount);
 
-	bool bSucceeded = true;
+	bool bAnyViewSucceeded = false;
 
 	for (uint32_t View = 0; View < ViewCount; ++View)
 	{
-		if (!BuildOcclusionMesh(XR_VISIBILITY_MASK_TYPE_VISIBLE_TRIANGLE_MESH_KHR, View, VisibleAreaMeshes[View]) ||
-			!BuildOcclusionMesh(XR_VISIBILITY_MASK_TYPE_HIDDEN_TRIANGLE_MESH_KHR, View, HiddenAreaMeshes[View]))
+		if (BuildOcclusionMesh(XR_VISIBILITY_MASK_TYPE_VISIBLE_TRIANGLE_MESH_KHR, View, VisibleAreaMeshes[View]) &&
+			BuildOcclusionMesh(XR_VISIBILITY_MASK_TYPE_HIDDEN_TRIANGLE_MESH_KHR, View, HiddenAreaMeshes[View]))
 		{
-			bSucceeded = false;
-			break;
+			bAnyViewSucceeded = true;
 		}
 	}
 
-	if (!bSucceeded)
+	if (!bAnyViewSucceeded)
 	{
 		UE_LOG(LogHMD, Error, TEXT("Failed to create all visibility mask meshes for device/views. Abandoning visibility mask."));
 
@@ -1844,9 +1886,7 @@ bool FOpenXRHMD::BuildOcclusionMesh(XrVisibilityMaskTypeKHR Type, int View, FHMD
 
 	if (VisibilityMask.indexCountOutput == 0)
 	{
-		UE_LOG(LogHMD, Warning, TEXT("Runtime does not currently have a visibility mask.  Another attempt will be made to build the occlusion mesh on an XR_TYPE_EVENT_DATA_VISIBILITY_MASK_CHANGED_KHR event."));
-		// Disallow future BuildOcclusionMesh attempts until this flag is reset in the XR_TYPE_EVENT_DATA_VISIBILITY_MASK_CHANGED_KHR handler.
-		bHiddenAreaMaskSupported = false;
+		// Runtime doesn't have a valid mask for this view
 		return false;
 	}
 	if (!VisibilityMask.indexCountOutput || (VisibilityMask.indexCountOutput % 3) != 0 || VisibilityMask.vertexCountOutput == 0)
@@ -1914,6 +1954,8 @@ bool FOpenXRHMD::BuildOcclusionMesh(XrVisibilityMaskTypeKHR Type, int View, FHMD
 bool FOpenXRHMD::OnStereoStartup()
 {
 	FWriteScopeLock Lock(SessionHandleMutex);
+	FWriteScopeLock DeviceLock(DeviceMutex);
+
 	XrSessionCreateInfo SessionInfo;
 	SessionInfo.type = XR_TYPE_SESSION_CREATE_INFO;
 	SessionInfo.next = RenderBridge->GetGraphicsBinding();
@@ -2041,14 +2083,14 @@ bool FOpenXRHMD::OnStereoTeardown()
 
 void FOpenXRHMD::DestroySession()
 {
-	FWriteScopeLock DeviceLock(DeviceMutex);
-	
 	// FlushRenderingCommands must be called outside of SessionLock since some rendering threads will also lock this mutex.
 	FlushRenderingCommands();
 	FWriteScopeLock SessionLock(SessionHandleMutex);
 
 	if (Session != XR_NULL_HANDLE)
 	{
+		FWriteScopeLock DeviceLock(DeviceMutex);
+
 		// We need to reset all swapchain references to ensure there are no attempts
 		// to destroy swapchain handles after the session is already destroyed.
 		ForEachLayer([&](uint32 /* unused */, FOpenXRLayer& Layer)
@@ -2068,6 +2110,11 @@ void FOpenXRHMD::DestroySession()
 		PipelinedLayerStateRHI.DepthSwapchain.Reset();
 		PipelinedLayerStateRHI.QuadSwapchains.Reset();
 
+		// Reset the frame state.
+		PipelinedFrameStateGame.FrameState = XrFrameState{ XR_TYPE_FRAME_STATE };
+		PipelinedFrameStateRendering.FrameState = XrFrameState{ XR_TYPE_FRAME_STATE };
+		PipelinedFrameStateRHI.FrameState = XrFrameState{ XR_TYPE_FRAME_STATE };
+
 		// VRFocus must be reset so FWindowsApplication::PollGameDeviceState does not incorrectly short-circuit.
 		FApp::SetUseVRFocus(false);
 		FApp::SetHasVRFocus(false);
@@ -2095,6 +2142,8 @@ void FOpenXRHMD::DestroySession()
 
 int32 FOpenXRHMD::AddActionDevice(XrAction Action, XrPath Path)
 {
+	FWriteScopeLock DeviceLock(DeviceMutex);
+
 	// Ensure the HMD device is already emplaced
 	ensure(DeviceSpaces.Num() > 0);
 
@@ -2111,6 +2160,8 @@ int32 FOpenXRHMD::AddActionDevice(XrAction Action, XrPath Path)
 
 void FOpenXRHMD::ResetActionDevices()
 {
+	FWriteScopeLock DeviceLock(DeviceMutex);
+
 	// Index 0 is HMDDeviceId and is preserved. The remaining are action devices.
 	if (DeviceSpaces.Num() > 0)
 	{
@@ -2120,6 +2171,7 @@ void FOpenXRHMD::ResetActionDevices()
 
 XrPath FOpenXRHMD::GetTrackedDevicePath(const int32 DeviceId)
 {
+	FReadScopeLock DeviceLock(DeviceMutex);
 	if (DeviceSpaces.IsValidIndex(DeviceId))
 	{
 		return DeviceSpaces[DeviceId].Path;
@@ -2234,6 +2286,9 @@ bool FOpenXRHMD::AllocateRenderTargetTexture(uint32 Index, uint32 SizeX, uint32
 			return false;
 		}
 
+		// Acquire the first swapchain image
+		Swapchain->IncrementSwapChainIndex_RHIThread();
+
 #if WITH_EDITOR
 		if (GIsEditor)
 		{
@@ -2284,6 +2339,9 @@ bool FOpenXRHMD::AllocateDepthTexture(uint32 Index, uint32 SizeX, uint32 SizeY,
 		{
 			return false;
 		}
+
+		// Acquire the first swapchain image
+		Swapchain->IncrementSwapChainIndex_RHIThread();
 	}
 
 	bNeedReAllocatedDepth = false;
@@ -2297,6 +2355,7 @@ bool FOpenXRHMD::AllocateDepthTexture(uint32 Index, uint32 SizeX, uint32 SizeY,
 void FOpenXRHMD::OnBeginRendering_RenderThread(FRHICommandListImmediate& RHICmdList, FSceneViewFamily& ViewFamily)
 {
 	ensure(IsInRenderingThread());
+	FReadScopeLock DeviceLock(DeviceMutex);
 
 	const float WorldToMeters = GetWorldToMetersScale();
 	const FTransform InvTrackingToWorld = GetTrackingToWorldTransform().Inverse();
@@ -2317,6 +2376,53 @@ void FOpenXRHMD::OnBeginRendering_RenderThread(FRHICommandListImmediate& RHICmdL
 		Projection.fov = View.fov;
 	}
 
+	// Manage the swapchains in the stereo layers
+	auto CreateSwapchain = [this](FRHITexture2D* Texture, ETextureCreateFlags Flags)
+	{
+		return RenderBridge->CreateSwapchain(Session,
+			PF_B8G8R8A8,
+			Texture->GetSizeX(),
+			Texture->GetSizeY(),
+			1,
+			Texture->GetNumMips(),
+			Texture->GetNumSamples(),
+			Texture->GetFlags() | Flags,
+			TexCreate_RenderTargetable,
+			Texture->GetClearBinding());
+	};
+	if (GetStereoLayersDirty())
+	{
+		ForEachLayer([&](uint32 /* unused */, FOpenXRLayer& Layer)
+		{
+			if (Layer.Desc.IsVisible() && Layer.Desc.HasShape<FQuadLayer>())
+			{
+				const ETextureCreateFlags Flags = Layer.Desc.Flags & IStereoLayers::LAYER_FLAG_TEX_CONTINUOUS_UPDATE ?
+					TexCreate_Dynamic | TexCreate_SRGB : TexCreate_SRGB;
+
+				if (Layer.NeedReAllocateTexture())
+				{
+					FRHITexture2D* Texture = Layer.Desc.Texture->GetTexture2D();
+					Layer.Swapchain = CreateSwapchain(Texture, Flags);
+					Layer.SwapchainSize = Texture->GetSizeXY();
+					Layer.bUpdateTexture = true;
+				}
+
+				if (Layer.NeedReAllocateLeftTexture())
+				{
+					FRHITexture2D* Texture = Layer.Desc.LeftTexture->GetTexture2D();
+					Layer.LeftSwapchain = CreateSwapchain(Texture, Flags);
+					Layer.bUpdateTexture = true;
+				}
+			}
+			else
+			{
+				// We retain references in FPipelinedLayerState to avoid premature destruction
+				Layer.Swapchain.Reset();
+				Layer.LeftSwapchain.Reset();
+			}
+		});
+	}
+
 	// Gather all active quad layers for composition sorted by priority
 	TArray<FOpenXRLayer> StereoLayers;
 	CopySortedLayers(StereoLayers);
@@ -2390,7 +2496,7 @@ void FOpenXRHMD::OnBeginRendering_RenderThread(FRHICommandListImmediate& RHICmdL
 		ForEachLayer([&](uint32 /* unused */, FOpenXRLayer& Layer)
 		{
 			Layer.bUpdateTexture = Layer.Desc.Flags & IStereoLayers::LAYER_FLAG_TEX_CONTINUOUS_UPDATE;
-		});
+		}, false);
 
 		{
 			FReadScopeLock Lock(SessionHandleMutex);
@@ -2400,9 +2506,11 @@ void FOpenXRHMD::OnBeginRendering_RenderThread(FRHICommandListImmediate& RHICmdL
 			}
 		}
 
-		RHICmdList.EnqueueLambda([this](FRHICommandListImmediate& InRHICmdList)
+		FXRSwapChainPtr ColorSwapchain = PipelinedLayerStateRendering.ColorSwapchain;
+		FXRSwapChainPtr DepthSwapchain = PipelinedLayerStateRendering.DepthSwapchain;
+		RHICmdList.EnqueueLambda([this, FrameState = PipelinedFrameStateRendering, ColorSwapchain, DepthSwapchain](FRHICommandListImmediate& InRHICmdList)
 		{
-			OnBeginRendering_RHIThread();
+			OnBeginRendering_RHIThread(FrameState, ColorSwapchain, DepthSwapchain);
 		});
 	}
 
@@ -2413,6 +2521,7 @@ void FOpenXRHMD::OnBeginRendering_RenderThread(FRHICommandListImmediate& RHICmdL
 void FOpenXRHMD::LocateViews(FPipelinedFrameState& PipelineState, bool ResizeViewsArray)
 {
 	check(PipelineState.FrameState.predictedDisplayTime);
+	FReadScopeLock DeviceLock(DeviceMutex);
 
 	uint32_t ViewCount = 0;
 	XrViewLocateInfo ViewInfo;
@@ -2460,9 +2569,9 @@ void FOpenXRHMD::OnLateUpdateApplied_RenderThread(FRHICommandListImmediate& RHIC
 		Projection.fov = View.fov;
 	}
 
-	RHICmdList.EnqueueLambda([this](FRHICommandListImmediate& InRHICmdList)
+	RHICmdList.EnqueueLambda([this, ProjectionLayers = PipelinedLayerStateRendering.ProjectionLayers](FRHICommandListImmediate& InRHICmdList)
 	{
-		PipelinedLayerStateRHI.ProjectionLayers = PipelinedLayerStateRendering.ProjectionLayers;
+		PipelinedLayerStateRHI.ProjectionLayers = ProjectionLayers;
 	});
 }
 
@@ -2646,7 +2755,7 @@ bool FOpenXRHMD::OnStartGameFrame(FWorldContext& WorldContext)
 	return true;
 }
 
-void FOpenXRHMD::OnBeginRendering_RHIThread()
+void FOpenXRHMD::OnBeginRendering_RHIThread(const FPipelinedFrameState& InFrameState, FXRSwapChainPtr ColorSwapchain, FXRSwapChainPtr DepthSwapchain)
 {
 	ensure(IsInRenderingThread() || IsInRHIThread());
 
@@ -2661,13 +2770,13 @@ void FOpenXRHMD::OnBeginRendering_RHIThread()
 		return;
 	}
 
-	PipelinedFrameStateRHI = PipelinedFrameStateRendering;
-	PipelinedLayerStateRHI = PipelinedLayerStateRendering;
+	// The layer state will be copied after SetFinalViewRect
+	PipelinedFrameStateRHI = InFrameState;
 
 	XrFrameBeginInfo BeginInfo;
 	BeginInfo.type = XR_TYPE_FRAME_BEGIN_INFO;
 	BeginInfo.next = nullptr;
-	XrTime DisplayTime = PipelinedFrameStateRendering.FrameState.predictedDisplayTime;
+	XrTime DisplayTime = InFrameState.FrameState.predictedDisplayTime;
 	for (IOpenXRExtensionPlugin* Module : ExtensionPlugins)
 	{
 		BeginInfo.next = Module->OnBeginFrame(Session, DisplayTime, BeginInfo.next);
@@ -2679,18 +2788,18 @@ void FOpenXRHMD::OnBeginRendering_RHIThread()
 		// Only the swapchains are valid to pull out of PipelinedLayerStateRendering
 		// Full population is deferred until SetFinalViewRect.
 		// TODO Possibly move these Waits to SetFinalViewRect??
-		PipelinedLayerStateRHI.ColorSwapchain = PipelinedLayerStateRendering.ColorSwapchain;
-		PipelinedLayerStateRHI.DepthSwapchain = PipelinedLayerStateRendering.DepthSwapchain;
+		PipelinedLayerStateRHI.ColorSwapchain = ColorSwapchain;
+		PipelinedLayerStateRHI.DepthSwapchain = DepthSwapchain;
 
 		// We need a new swapchain image unless we've already acquired one for rendering
-		if (!bIsRendering && PipelinedLayerStateRHI.ColorSwapchain)
+		if (!bIsRendering && ColorSwapchain)
 		{
-			PipelinedLayerStateRHI.ColorSwapchain->IncrementSwapChainIndex_RHIThread();
-			PipelinedLayerStateRHI.ColorSwapchain->WaitCurrentImage_RHIThread(OPENXR_SWAPCHAIN_WAIT_TIMEOUT);
-			if (bDepthExtensionSupported && PipelinedLayerStateRHI.DepthSwapchain)
+			ColorSwapchain->IncrementSwapChainIndex_RHIThread();
+			ColorSwapchain->WaitCurrentImage_RHIThread(OPENXR_SWAPCHAIN_WAIT_TIMEOUT);
+			if (bDepthExtensionSupported && DepthSwapchain)
 			{
-				PipelinedLayerStateRHI.DepthSwapchain->IncrementSwapChainIndex_RHIThread();
-				PipelinedLayerStateRHI.DepthSwapchain->WaitCurrentImage_RHIThread(OPENXR_SWAPCHAIN_WAIT_TIMEOUT);
+				DepthSwapchain->IncrementSwapChainIndex_RHIThread();
+				DepthSwapchain->WaitCurrentImage_RHIThread(OPENXR_SWAPCHAIN_WAIT_TIMEOUT);
 			}
 		}
 
@@ -2956,80 +3065,38 @@ bool FOpenXRHMD::HasVisibleAreaMesh() const
 
 void FOpenXRHMD::DrawHiddenAreaMesh_RenderThread(class FRHICommandList& RHICmdList, EStereoscopicPass StereoPass) const
 {
-	check(IsInRenderingThread());
 	check(StereoPass != eSSP_FULL);
 
 	const uint32 ViewIndex = GetViewIndexForPass(StereoPass);
 	if (ViewIndex < (uint32)HiddenAreaMeshes.Num())
 	{
 		const FHMDViewMesh& Mesh = HiddenAreaMeshes[ViewIndex];
-		check(Mesh.IsValid());
 
-		RHICmdList.SetStreamSource(0, Mesh.VertexBufferRHI, 0);
-		RHICmdList.DrawIndexedPrimitive(Mesh.IndexBufferRHI, 0, 0, Mesh.NumVertices, 0, Mesh.NumTriangles, 1);
+		if (Mesh.IsValid())
+		{
+			RHICmdList.SetStreamSource(0, Mesh.VertexBufferRHI, 0);
+			RHICmdList.DrawIndexedPrimitive(Mesh.IndexBufferRHI, 0, 0, Mesh.NumVertices, 0, Mesh.NumTriangles, 1);
+		}
 	}
 }
 
 void FOpenXRHMD::DrawVisibleAreaMesh_RenderThread(class FRHICommandList& RHICmdList, EStereoscopicPass StereoPass) const
 {
-	check(IsInRenderingThread());
 	check(StereoPass != eSSP_FULL);
 
 	const uint32 ViewIndex = GetViewIndexForPass(StereoPass);
 	check(ViewIndex < (uint32)VisibleAreaMeshes.Num());
 	const FHMDViewMesh& Mesh = VisibleAreaMeshes[ViewIndex];
-	check(Mesh.IsValid());
-
-	RHICmdList.SetStreamSource(0, Mesh.VertexBufferRHI, 0);
-	RHICmdList.DrawIndexedPrimitive(Mesh.IndexBufferRHI, 0, 0, Mesh.NumVertices, 0, Mesh.NumTriangles, 1);
-}
 
-void FOpenXRHMD::UpdateLayer(FOpenXRLayer& Layer, uint32 LayerId, bool bIsValid) const
-{
-	if (!Layer.Desc.HasShape<FQuadLayer>())
+	if (Mesh.IsValid())
 	{
-		return;
-	}
-
-	if (bIsValid)
-	{
-		const ETextureCreateFlags Flags = Layer.Desc.Flags & IStereoLayers::LAYER_FLAG_TEX_CONTINUOUS_UPDATE ?
-			TexCreate_Dynamic | TexCreate_SRGB : TexCreate_SRGB;
-
-		auto CreateSwapchain = [this](FRHITexture2D* Texture, ETextureCreateFlags Flags)
-		{
-			return RenderBridge->CreateSwapchain(Session,
-				PF_B8G8R8A8,
-				Texture->GetSizeX(),
-				Texture->GetSizeY(),
-				1,
-				Texture->GetNumMips(),
-				Texture->GetNumSamples(),
-				Texture->GetFlags() | Flags,
-				TexCreate_RenderTargetable,
-				Texture->GetClearBinding());
-		};
-
-		if (Layer.NeedReAllocateTexture())
-		{
-			FRHITexture2D* Texture = Layer.Desc.Texture->GetTexture2D();
-			Layer.Swapchain = CreateSwapchain(Texture, Flags);
-			Layer.SwapchainSize = Texture->GetSizeXY();
-			Layer.bUpdateTexture = true;
-		}
-
-		if (Layer.NeedReAllocateLeftTexture())
-		{
-			FRHITexture2D* Texture = Layer.Desc.LeftTexture->GetTexture2D();
-			Layer.LeftSwapchain = CreateSwapchain(Texture, Flags);
-			Layer.bUpdateTexture = true;
-		}
+		RHICmdList.SetStreamSource(0, Mesh.VertexBufferRHI, 0);
+		RHICmdList.DrawIndexedPrimitive(Mesh.IndexBufferRHI, 0, 0, Mesh.NumVertices, 0, Mesh.NumTriangles, 1);
 	}
 	else
 	{
-		// We retain references in FPipelinedLayerState to avoid premature destruction
-		Layer.Swapchain.Reset();
-		Layer.LeftSwapchain.Reset();
+		// Invalid mesh means that entire area is visible, draw a fullscreen quad to simulate
+		FPixelShaderUtils::DrawFullscreenQuad(RHICmdList, 1);
 	}
 }
 
diff --git a/Engine/Plugins/Runtime/OpenXR/Source/OpenXRHMD/Private/OpenXRHMD.h b/Engine/Plugins/Runtime/OpenXR/Source/OpenXRHMD/Private/OpenXRHMD.h
--- a/Engine/Plugins/Runtime/OpenXR/Source/OpenXRHMD/Private/OpenXRHMD.h
+++ b/Engine/Plugins/Runtime/OpenXR/Source/OpenXRHMD/Private/OpenXRHMD.h
@@ -227,7 +227,7 @@ public:
 	virtual bool IsStereoEnabled() const override;
 	virtual bool EnableStereo(bool stereo = true) override;
 	virtual void AdjustViewRect(EStereoscopicPass StereoPass, int32& X, int32& Y, uint32& SizeX, uint32& SizeY) const override;
-	virtual void SetFinalViewRect(const enum EStereoscopicPass StereoPass, const FIntRect& FinalViewRect) override;
+	virtual void SetFinalViewRect(FRHICommandListImmediate& RHICmdList, const enum EStereoscopicPass StereoPass, const FIntRect& FinalViewRect) override;
 	virtual int32 GetDesiredNumberOfViews(bool bStereoRequested) const override;
 	virtual EStereoscopicPass GetViewPassForIndex(bool bStereoRequested, uint32 ViewIndex) const override;
 	virtual uint32 GetViewIndexForPass(EStereoscopicPass StereoPassType) const override;
@@ -260,7 +260,6 @@ public:
 	virtual void OnEndPlay(FWorldContext& InWorldContext) override;
 
 	/** IStereoLayers */
-	virtual void UpdateLayer(FOpenXRLayer& Layer, uint32 LayerId, bool bIsValid) const override;
 	virtual bool ShouldCopyDebugLayersToSpectatorScreen() const override { return true; }
 
 public:
@@ -270,7 +269,7 @@ public:
 	/** Destructor */
 	virtual ~FOpenXRHMD();
 
-	void OnBeginRendering_RHIThread();
+	void OnBeginRendering_RHIThread(const FPipelinedFrameState& InFrameState, FXRSwapChainPtr ColorSwapchain, FXRSwapChainPtr DepthSwapchain);
 	void OnFinishRendering_RHIThread();
 
 	/** @return	True if the HMD was initialized OK */
diff --git a/Engine/Plugins/Runtime/OpenXR/Source/OpenXRInput/Private/OpenXRInput.cpp b/Engine/Plugins/Runtime/OpenXR/Source/OpenXRInput/Private/OpenXRInput.cpp
--- a/Engine/Plugins/Runtime/OpenXR/Source/OpenXRInput/Private/OpenXRInput.cpp
+++ b/Engine/Plugins/Runtime/OpenXR/Source/OpenXRInput/Private/OpenXRInput.cpp
@@ -568,12 +568,7 @@ namespace OpenXRInputNamespace
 
 void FOpenXRInputPlugin::FOpenXRInput::SendControllerEvents()
 {
-	if (!bActionsBound)
-	{
-		return;
-	}
-
-	if (OpenXRHMD == nullptr)
+	if (!bActionsBound || OpenXRHMD == nullptr)
 	{
 		return;
 	}
@@ -890,7 +885,7 @@ void FOpenXRInputPlugin::FOpenXRInput::EnumerateSources(TArray<FMotionController
 // TODO: Refactor API to change the Hand type to EControllerHand
 void FOpenXRInputPlugin::FOpenXRInput::SetHapticFeedbackValues(int32 ControllerId, int32 Hand, const FHapticFeedbackValues& Values)
 {
-	if (OpenXRHMD == nullptr)
+	if (!bActionsBound || OpenXRHMD == nullptr)
 	{
 		return;
 	}
diff --git a/Engine/Source/Runtime/Engine/Public/StereoRendering.h b/Engine/Source/Runtime/Engine/Public/StereoRendering.h
--- a/Engine/Source/Runtime/Engine/Public/StereoRendering.h
+++ b/Engine/Source/Runtime/Engine/Public/StereoRendering.h
@@ -195,7 +195,7 @@ public:
 	/**
 	* Provides the final view rect that the renderer will render into.
 	*/
-	virtual void SetFinalViewRect(const enum EStereoscopicPass StereoPass, const FIntRect& FinalViewRect) {}
+	virtual void SetFinalViewRect(FRHICommandListImmediate& RHICmdList, const enum EStereoscopicPass StereoPass, const FIntRect& FinalViewRect) {}
 
 	/**
 	 * Gets the percentage bounds of the safe region to draw in.  This allows things like stat rendering to appear within the readable portion of the stereo view.
diff --git a/Engine/Source/Runtime/Renderer/Private/DeferredShadingRenderer.cpp b/Engine/Source/Runtime/Renderer/Private/DeferredShadingRenderer.cpp
--- a/Engine/Source/Runtime/Renderer/Private/DeferredShadingRenderer.cpp
+++ b/Engine/Source/Runtime/Renderer/Private/DeferredShadingRenderer.cpp
@@ -1915,7 +1915,7 @@ void FDeferredShadingSceneRenderer::Render(FRDGBuilder& GraphBuilder)
 	CSV_SCOPED_TIMING_STAT_EXCLUSIVE(RenderOther);
 
 	// Setups the final FViewInfo::ViewRect.
-	PrepareViewRectsForRendering();
+	PrepareViewRectsForRendering(GraphBuilder.RHICmdList);
 
 	if (ShouldRenderSkyAtmosphere(Scene, ViewFamily.EngineShowFlags))
 	{
diff --git a/Engine/Source/Runtime/Renderer/Private/MobileShadingRenderer.cpp b/Engine/Source/Runtime/Renderer/Private/MobileShadingRenderer.cpp
--- a/Engine/Source/Runtime/Renderer/Private/MobileShadingRenderer.cpp
+++ b/Engine/Source/Runtime/Renderer/Private/MobileShadingRenderer.cpp
@@ -664,7 +664,7 @@ void FMobileSceneRenderer::Render(FRDGBuilder& GraphBuilder)
 	// Establish scene primitive count (must be done after UpdateAllPrimitiveSceneInfos)
 	FGPUSceneScopeBeginEndHelper GPUSceneScopeBeginEndHelper(Scene->GPUScene, GPUSceneDynamicContext, Scene);
 
-	PrepareViewRectsForRendering();
+	PrepareViewRectsForRendering(GraphBuilder.RHICmdList);
 
 	if (ShouldRenderSkyAtmosphere(Scene, ViewFamily.EngineShowFlags))
 	{
diff --git a/Engine/Source/Runtime/Renderer/Private/PlanarReflectionRendering.cpp b/Engine/Source/Runtime/Renderer/Private/PlanarReflectionRendering.cpp
--- a/Engine/Source/Runtime/Renderer/Private/PlanarReflectionRendering.cpp
+++ b/Engine/Source/Runtime/Renderer/Private/PlanarReflectionRendering.cpp
@@ -457,7 +457,7 @@ static void UpdatePlanarReflectionContentsWithoutRendering_RenderThread(
 			}
 		}
 
-		SceneRenderer->PrepareViewRectsForRendering();
+		SceneRenderer->PrepareViewRectsForRendering(RHICmdList);
 
 		SceneProxy->RenderTarget = RenderTarget;
 
diff --git a/Engine/Source/Runtime/Renderer/Private/SceneHitProxyRendering.cpp b/Engine/Source/Runtime/Renderer/Private/SceneHitProxyRendering.cpp
--- a/Engine/Source/Runtime/Renderer/Private/SceneHitProxyRendering.cpp
+++ b/Engine/Source/Runtime/Renderer/Private/SceneHitProxyRendering.cpp
@@ -542,7 +542,7 @@ void FMobileSceneRenderer::RenderHitProxies(FRDGBuilder& GraphBuilder)
 
 	FGPUSceneScopeBeginEndHelper GPUSceneScopeBeginEndHelper(Scene->GPUScene, GPUSceneDynamicContext, Scene);
 
-	PrepareViewRectsForRendering();
+	PrepareViewRectsForRendering(GraphBuilder.RHICmdList);
 
 #if WITH_EDITOR
 	FSceneTexturesConfig SceneTexturesConfig = FSceneTexturesConfig::Create(ViewFamily);
@@ -578,7 +578,7 @@ void FDeferredShadingSceneRenderer::RenderHitProxies(FRDGBuilder& GraphBuilder)
 
 	FGPUSceneScopeBeginEndHelper GPUSceneScopeBeginEndHelper(Scene->GPUScene, GPUSceneDynamicContext, Scene);
 
-	PrepareViewRectsForRendering();
+	PrepareViewRectsForRendering(GraphBuilder.RHICmdList);
 
 #if WITH_EDITOR
 	const FSceneTexturesConfig SceneTexturesConfig = FSceneTexturesConfig::Create(ViewFamily);
diff --git a/Engine/Source/Runtime/Renderer/Private/SceneRendering.cpp b/Engine/Source/Runtime/Renderer/Private/SceneRendering.cpp
--- a/Engine/Source/Runtime/Renderer/Private/SceneRendering.cpp
+++ b/Engine/Source/Runtime/Renderer/Private/SceneRendering.cpp
@@ -2595,7 +2595,7 @@ FIntPoint FSceneRenderer::GetDesiredInternalBufferSize(const FSceneViewFamily& V
 }
 
 
-void FSceneRenderer::PrepareViewRectsForRendering()
+void FSceneRenderer::PrepareViewRectsForRendering(FRHICommandListImmediate& RHICmdList)
 {
 	check(IsInRenderingThread());
 
@@ -2616,7 +2616,7 @@ void FSceneRenderer::PrepareViewRectsForRendering()
 			for (int32 i = 0; i < Views.Num(); i++)
 			{
 				FViewInfo& View = Views[i];
-				GEngine->StereoRenderingDevice->SetFinalViewRect(View.StereoPass, View.ViewRect);
+				GEngine->StereoRenderingDevice->SetFinalViewRect(RHICmdList, View.StereoPass, View.ViewRect);
 			}
 		}
 		return;
@@ -2771,7 +2771,7 @@ void FSceneRenderer::PrepareViewRectsForRendering()
 		for (int32 i = 0; i < Views.Num(); i++)
 		{
 			FViewInfo& View = Views[i];
-			GEngine->StereoRenderingDevice->SetFinalViewRect(View.StereoPass, View.ViewRect);
+			GEngine->StereoRenderingDevice->SetFinalViewRect(RHICmdList, View.StereoPass, View.ViewRect);
 		}
 	}
 }
diff --git a/Engine/Source/Runtime/Renderer/Private/SceneRendering.h b/Engine/Source/Runtime/Renderer/Private/SceneRendering.h
--- a/Engine/Source/Runtime/Renderer/Private/SceneRendering.h
+++ b/Engine/Source/Runtime/Renderer/Private/SceneRendering.h
@@ -1869,7 +1869,7 @@ public:
 	static FSceneRenderer* CreateSceneRenderer(const FSceneViewFamily* InViewFamily, FHitProxyConsumer* HitProxyConsumer);
 
 	/** Setups FViewInfo::ViewRect according to ViewFamilly's ScreenPercentageInterface. */
-	void PrepareViewRectsForRendering();
+	void PrepareViewRectsForRendering(FRHICommandListImmediate& RHICmdList);
 
 #if WITH_MGPU
 	/** Assigns the view GPU masks and returns the render target GPU mask. */
