Fix naming (UE4 -> UE, Unreal Engine 4 -> Unreal Engine).

diff --git a/Engine/Plugins/Runtime/OpenXR/Source/OpenXRHMD/Private/OpenXRHMD.cpp b/Engine/Plugins/Runtime/OpenXR/Source/OpenXRHMD/Private/OpenXRHMD.cpp
--- a/Engine/Plugins/Runtime/OpenXR/Source/OpenXRHMD/Private/OpenXRHMD.cpp
+++ b/Engine/Plugins/Runtime/OpenXR/Source/OpenXRHMD/Private/OpenXRHMD.cpp
@@ -2318,7 +2318,7 @@ bool FOpenXRHMD::AllocateDepthTexture(uint32 Index, uint32 SizeX, uint32 SizeY,
 {
 	check(IsInRenderingThread());
 
-	// FIXME: UE4 constantly calls this function even when there is no reason to reallocate the depth texture
+	// FIXME: UE constantly calls this function even when there is no reason to reallocate the depth texture
 	FReadScopeLock Lock(SessionHandleMutex);
 	if (!Session || !bDepthExtensionSupported)
 	{
diff --git a/Engine/Plugins/Runtime/OpenXR/Source/OpenXRHMD/Private/OpenXRHMD_Swapchain.cpp b/Engine/Plugins/Runtime/OpenXR/Source/OpenXRHMD/Private/OpenXRHMD_Swapchain.cpp
--- a/Engine/Plugins/Runtime/OpenXR/Source/OpenXRHMD/Private/OpenXRHMD_Swapchain.cpp
+++ b/Engine/Plugins/Runtime/OpenXR/Source/OpenXRHMD/Private/OpenXRHMD_Swapchain.cpp
@@ -206,7 +206,7 @@ FXRSwapChainPtr CreateSwapchain_D3D11(XrSession InSession, uint8 Format, uint32
 		DXGI_FORMAT PlatformFormat = (DXGI_FORMAT)GPixelFormats[InFormat].PlatformFormat;
 		PlatformFormat = ::FindDepthStencilDXGIFormat(PlatformFormat);
 
-		// UE4 renders a gamma-corrected image so we need to use an sRGB format if available
+		// UE renders a gamma-corrected image so we need to use an sRGB format if available
 		PlatformFormat = ::FindShaderResourceDXGIFormat(PlatformFormat, true);
 		return PlatformFormat;
 	};
@@ -244,7 +244,7 @@ FXRSwapChainPtr CreateSwapchain_D3D12(XrSession InSession, uint8 Format, uint32
 		DXGI_FORMAT PlatformFormat = (DXGI_FORMAT)GPixelFormats[InFormat].PlatformFormat;
 		PlatformFormat = D3D12RHI::FindDepthStencilDXGIFormat(PlatformFormat);
 
-		// UE4 renders a gamma-corrected image so we need to use an sRGB format if available
+		// UE renders a gamma-corrected image so we need to use an sRGB format if available
 		PlatformFormat = D3D12RHI::FindShaderResourceDXGIFormat(PlatformFormat, true);
 		return PlatformFormat;
 	};
@@ -307,7 +307,7 @@ FXRSwapChainPtr CreateSwapchain_Vulkan(XrSession InSession, uint8 Format, uint32
 {
 	TFunction<uint32(uint8)> ToPlatformFormat = [](uint8 InFormat)
 	{
-		// UE4 renders a gamma-corrected image so we need to use an sRGB format if available
+		// UE renders a gamma-corrected image so we need to use an sRGB format if available
 		return UEToVkTextureFormat(GPixelFormats[InFormat].UnrealFormat, true);
 	};
 	Format = FOpenXRSwapchain::GetNearestSupportedSwapchainFormat(InSession, Format, ToPlatformFormat);
diff --git a/Engine/Plugins/Runtime/OpenXR/Source/OpenXRInput/Private/OpenXRInput.cpp b/Engine/Plugins/Runtime/OpenXR/Source/OpenXRInput/Private/OpenXRInput.cpp
--- a/Engine/Plugins/Runtime/OpenXR/Source/OpenXRInput/Private/OpenXRInput.cpp
+++ b/Engine/Plugins/Runtime/OpenXR/Source/OpenXRInput/Private/OpenXRInput.cpp
@@ -252,8 +252,8 @@ void FOpenXRInputPlugin::FOpenXRInput::BuildActions()
 	XrActionSetCreateInfo SetInfo;
 	SetInfo.type = XR_TYPE_ACTION_SET_CREATE_INFO;
 	SetInfo.next = nullptr;
-	FCStringAnsi::Strcpy(SetInfo.actionSetName, XR_MAX_ACTION_SET_NAME_SIZE, "ue4");
-	FCStringAnsi::Strcpy(SetInfo.localizedActionSetName, XR_MAX_ACTION_SET_NAME_SIZE, "Unreal Engine 4");
+	FCStringAnsi::Strcpy(SetInfo.actionSetName, XR_MAX_ACTION_SET_NAME_SIZE, "ue");
+	FCStringAnsi::Strcpy(SetInfo.localizedActionSetName, XR_MAX_ACTION_SET_NAME_SIZE, "Unreal Engine");
 	XR_ENSURE(xrCreateActionSet(Instance, &SetInfo, &ActionSet));
 
 	XrPath LeftHand = GetPath(Instance, "/user/hand/left");
