Split OpenXRHMD module out of the OpenXRHMD implementation

This also renames IOpenXRHMDPlugin to IOpenXRHMDModule to avoid confusion with IOpenXRExtensionPlugin

diff --git a/Engine/Plugins/Runtime/Oculus/OculusOpenXR/Source/OculusOpenXRHMD/Private/OculusOpenXRHMD.cpp b/Engine/Plugins/Runtime/Oculus/OculusOpenXR/Source/OculusOpenXRHMD/Private/OculusOpenXRHMD.cpp
--- a/Engine/Plugins/Runtime/Oculus/OculusOpenXR/Source/OculusOpenXRHMD/Private/OculusOpenXRHMD.cpp
+++ b/Engine/Plugins/Runtime/Oculus/OculusOpenXR/Source/OculusOpenXRHMD/Private/OculusOpenXRHMD.cpp
@@ -101,7 +101,7 @@ void FOculusOpenXRHMD::OnEvent(XrSession InSession, const XrEventDataBaseHeader*
 	return;
 }
 
-const void* FOculusOpenXRHMD::OnCreateInstance(class IOpenXRHMDPlugin* InPlugin, const void* InNext)
+const void* FOculusOpenXRHMD::OnCreateInstance(class IOpenXRHMDModule* InModule, const void* InNext)
 {
 	//UE_LOG(LogOculusOpenXRPlugin, Log, TEXT("Oculus OpenXR OnCreateInstance"));
 	return InNext;
diff --git a/Engine/Plugins/Runtime/Oculus/OculusOpenXR/Source/OculusOpenXRHMD/Private/OculusOpenXRHMD.h b/Engine/Plugins/Runtime/Oculus/OculusOpenXR/Source/OculusOpenXRHMD/Private/OculusOpenXRHMD.h
--- a/Engine/Plugins/Runtime/Oculus/OculusOpenXR/Source/OculusOpenXRHMD/Private/OculusOpenXRHMD.h
+++ b/Engine/Plugins/Runtime/Oculus/OculusOpenXR/Source/OculusOpenXRHMD/Private/OculusOpenXRHMD.h
@@ -45,7 +45,7 @@ public:
 	virtual bool GetSpectatorScreenController(FHeadMountedDisplayBase* InHMDBase, TUniquePtr<FDefaultSpectatorScreenController>& OutSpectatorScreenController) override;
 	virtual void AddActions(XrInstance Instance, TFunction<XrAction(XrActionType InActionType, const FName& InName, const TArray<XrPath>& InSubactionPaths)> AddAction) override;
 	virtual void OnEvent(XrSession InSession, const XrEventDataBaseHeader* InHeader) override;
-	virtual const void* OnCreateInstance(class IOpenXRHMDPlugin* InPlugin, const void* InNext) override;
+	virtual const void* OnCreateInstance(class IOpenXRHMDModule* InModule, const void* InNext) override;
 	virtual const void* OnGetSystem(XrInstance InInstance, const void* InNext) override;
 	virtual const void* OnCreateSession(XrInstance InInstance, XrSystemId InSystem, const void* InNext) override;
 	virtual const void* OnBeginSession(XrSession InSession, const void* InNext) override;
diff --git a/Engine/Plugins/Runtime/OpenXR/Source/OpenXRAR/Private/OpenXRAR.h b/Engine/Plugins/Runtime/OpenXR/Source/OpenXRAR/Private/OpenXRAR.h
--- a/Engine/Plugins/Runtime/OpenXR/Source/OpenXRAR/Private/OpenXRAR.h
+++ b/Engine/Plugins/Runtime/OpenXR/Source/OpenXRAR/Private/OpenXRAR.h
@@ -3,7 +3,7 @@
 #pragma once
 
 #include "IOpenXRARTrackedGeometryHolder.h"
-#include "IOpenXRHMDPlugin.h"
+#include "IOpenXRHMDModule.h"
 #include "ARSystemSupportBase.h"
 #include "ARTraceResult.h"
 #include "ARPin.h"
diff --git a/Engine/Plugins/Runtime/OpenXR/Source/OpenXRHMD/Private/OpenXRHMD.cpp b/Engine/Plugins/Runtime/OpenXR/Source/OpenXRHMD/Private/OpenXRHMD.cpp
--- a/Engine/Plugins/Runtime/OpenXR/Source/OpenXRHMD/Private/OpenXRHMD.cpp
+++ b/Engine/Plugins/Runtime/OpenXR/Source/OpenXRHMD/Private/OpenXRHMD.cpp
@@ -6,7 +6,6 @@
 #include "OpenXRHMD_Swapchain.h"
 #include "OpenXRCore.h"
 #include "IOpenXRExtensionPlugin.h"
-#include "IOpenXRARModule.h"
 
 #include "Misc/App.h"
 #include "Misc/Parse.h"
@@ -15,7 +14,6 @@
 #include "Engine/Engine.h"
 #include "GameFramework/PlayerController.h"
 #include "Engine/LocalPlayer.h"
-#include "IOpenXRHMDPlugin.h"
 #include "SceneRendering.h"
 #include "PostProcess/PostProcessHMD.h"
 #include "GameFramework/WorldSettings.h"
@@ -26,17 +24,11 @@
 #include "PipelineStateCache.h"
 #include "Slate/SceneViewport.h"
 #include "Engine/GameEngine.h"
-#include "BuildSettings.h"
 #include "ARSystem.h"
 #include "IHandTracker.h"
 #include "PixelShaderUtils.h"
 #include "ScenePrivate.h"
 
-#if PLATFORM_ANDROID
-#include <android_native_app_glue.h>
-extern struct android_app* GNativeAndroidApp;
-#endif
-
 #if WITH_EDITOR
 #include "Editor/UnrealEd/Classes/Editor/EditorEngine.h"
 #endif
@@ -44,14 +36,6 @@ extern struct android_app* GNativeAndroidApp;
 #define OPENXR_PAUSED_IDLE_FPS 10
 static const int64 OPENXR_SWAPCHAIN_WAIT_TIMEOUT = 100000000ll;		// 100ms in nanoseconds.
 
-static TAutoConsoleVariable<int32> CVarEnableOpenXRValidationLayer(
-	TEXT("xr.EnableOpenXRValidationLayer"),
-	0,
-	TEXT("If true, enables the OpenXR validation layer, which will provide extended validation of\nOpenXR API calls. This should only be used for debugging purposes.\n")
-	TEXT("Changes will only take effect in new game/editor instances - can't be changed at runtime.\n"),
-	ECVF_Default);		// @todo: Should we specify ECVF_Cheat here so this doesn't show up in release builds?
-
-
 namespace {
 	static TSet<XrEnvironmentBlendMode> SupportedBlendModes{ XR_ENVIRONMENT_BLEND_MODE_ALPHA_BLEND, XR_ENVIRONMENT_BLEND_MODE_ADDITIVE, XR_ENVIRONMENT_BLEND_MODE_OPAQUE };
 	static TSet<XrViewConfigurationType> SupportedViewConfigurations{ XR_VIEW_CONFIGURATION_TYPE_PRIMARY_STEREO, XR_VIEW_CONFIGURATION_TYPE_PRIMARY_QUAD_VARJO };
@@ -90,677 +74,6 @@ namespace {
 	}
 }
 
-//---------------------------------------------------
-// OpenXRHMD Plugin Implementation
-//---------------------------------------------------
-
-class FOpenXRHMDPlugin : public IOpenXRHMDPlugin
-{
-public:
-	FOpenXRHMDPlugin()
-		: LoaderHandle(nullptr)
-		, Instance(XR_NULL_HANDLE)
-		, System(XR_NULL_SYSTEM_ID)
-		, RenderBridge(nullptr)
-	{ }
-
-	~FOpenXRHMDPlugin()
-	{
-	}
-
-	/** IHeadMountedDisplayModule implementation */
-	virtual TSharedPtr< class IXRTrackingSystem, ESPMode::ThreadSafe > CreateTrackingSystem() override;
-	virtual TSharedPtr< IHeadMountedDisplayVulkanExtensions, ESPMode::ThreadSafe > GetVulkanExtensions() override;
-	virtual uint64 GetGraphicsAdapterLuid() override;
-
-	FString GetModuleKeyName() const override
-	{
-		return FString(TEXT("OpenXRHMD"));
-	}
-
-	void GetModuleAliases(TArray<FString>& AliasesOut) const override
-	{
-		AliasesOut.Add(TEXT("OpenXR"));
-	}
-
-	void ShutdownModule() override
-	{
-		if (Instance)
-		{
-			XR_ENSURE(xrDestroyInstance(Instance));
-		}
-
-		if (LoaderHandle)
-		{
-			FPlatformProcess::FreeDllHandle(LoaderHandle);
-			LoaderHandle = nullptr;
-		}
-	}
-
-	virtual bool IsHMDConnected() override { return true; }
-	virtual bool IsStandaloneStereoOnlyDevice() override;
-	virtual bool IsExtensionAvailable(const FString& Name) const override { return AvailableExtensions.Contains(Name); }
-	virtual bool IsExtensionEnabled(const FString& Name) const override { return EnabledExtensions.Contains(Name); }
-	virtual bool IsLayerAvailable(const FString& Name) const override { return EnabledLayers.Contains(Name); }
-	virtual bool IsLayerEnabled(const FString& Name) const override { return EnabledLayers.Contains(Name); }
-
-private:
-	void *LoaderHandle;
-	XrInstance Instance;
-	XrSystemId System;
-	TSet<FString> AvailableExtensions;
-	TSet<FString> AvailableLayers;
-	TArray<const char*> EnabledExtensions;
-	TArray<const char*> EnabledLayers;
-	TArray<IOpenXRExtensionPlugin*> ExtensionPlugins;
-	TRefCountPtr<FOpenXRRenderBridge> RenderBridge;
-	TSharedPtr< IHeadMountedDisplayVulkanExtensions, ESPMode::ThreadSafe > VulkanExtensions;
-
-	bool EnumerateExtensions();
-	bool EnumerateLayers();
-	bool InitRenderBridge();
-	bool InitInstanceAndSystem();
-	bool InitInstance();
-	bool InitSystem();
-	PFN_xrGetInstanceProcAddr GetDefaultLoader();
-	bool EnableExtensions(const TArray<const ANSICHAR*>& RequiredExtensions, const TArray<const ANSICHAR*>& OptionalExtensions, TArray<const ANSICHAR*>& OutExtensions);
-	bool GetRequiredExtensions(TArray<const ANSICHAR*>& OutExtensions);
-	bool GetOptionalExtensions(TArray<const ANSICHAR*>& OutExtensions);
-};
-
-IMPLEMENT_MODULE( FOpenXRHMDPlugin, OpenXRHMD )
-
-TSharedPtr< class IXRTrackingSystem, ESPMode::ThreadSafe > FOpenXRHMDPlugin::CreateTrackingSystem()
-{
-	if (!RenderBridge)
-	{
-		if (!InitRenderBridge())
-		{
-			return nullptr;
-		}
-	}
-	auto ARModule = FModuleManager::LoadModulePtr<IOpenXRARModule>("OpenXRAR");
-	auto ARSystem = ARModule->CreateARSystem();
-
-	auto OpenXRHMD = FSceneViewExtensions::NewExtension<FOpenXRHMD>(Instance, System, RenderBridge, EnabledExtensions, ExtensionPlugins, ARSystem);
-	if (OpenXRHMD->IsInitialized())
-	{
-		ARModule->SetTrackingSystem(OpenXRHMD);
-		OpenXRHMD->GetARCompositionComponent()->InitializeARSystem();
-		return OpenXRHMD;
-	}
-
-	return nullptr;
-}
-
-uint64 FOpenXRHMDPlugin::GetGraphicsAdapterLuid()
-{
-	if (!RenderBridge)
-	{
-		if (!InitRenderBridge())
-		{
-			return 0;
-		}
-	}
-	return RenderBridge->GetGraphicsAdapterLuid();
-}
-
-TSharedPtr< IHeadMountedDisplayVulkanExtensions, ESPMode::ThreadSafe > FOpenXRHMDPlugin::GetVulkanExtensions()
-{
-#ifdef XR_USE_GRAPHICS_API_VULKAN
-	if (InitInstanceAndSystem() && IsExtensionEnabled(XR_KHR_VULKAN_ENABLE_EXTENSION_NAME))
-	{
-		if (!VulkanExtensions.IsValid())
-		{
-			VulkanExtensions = MakeShareable(new FOpenXRHMD::FVulkanExtensions(Instance, System));
-		}
-		return VulkanExtensions;
-	}
-#endif//XR_USE_GRAPHICS_API_VULKAN
-	return nullptr;
-}
-
-bool FOpenXRHMDPlugin::IsStandaloneStereoOnlyDevice()
-{
-	if (InitInstanceAndSystem())
-	{
-		for (IOpenXRExtensionPlugin* Module : ExtensionPlugins)
-		{
-			if (Module->IsStandaloneStereoOnlyDevice())
-			{
-				return true;
-			}
-		}
-	}
-	return false;
-}
-
-bool FOpenXRHMDPlugin::EnumerateExtensions()
-{
-	uint32_t ExtensionsCount = 0;
-	if (XR_FAILED(xrEnumerateInstanceExtensionProperties(nullptr, 0, &ExtensionsCount, nullptr)))
-	{
-		// If it fails this early that means there's no runtime installed
-		return false;
-	}
-
-	TArray<XrExtensionProperties> Properties;
-	Properties.SetNum(ExtensionsCount);
-	for (auto& Prop : Properties)
-	{
-		Prop = XrExtensionProperties{ XR_TYPE_EXTENSION_PROPERTIES };
-	}
-
-	if (XR_ENSURE(xrEnumerateInstanceExtensionProperties(nullptr, ExtensionsCount, &ExtensionsCount, Properties.GetData())))
-	{
-		for (const XrExtensionProperties& Prop : Properties)
-		{
-			AvailableExtensions.Add(Prop.extensionName);
-		}
-		return true;
-	}
-	return false;
-}
-
-bool FOpenXRHMDPlugin::EnumerateLayers()
-{
-	uint32 LayerPropertyCount = 0;
-	if (XR_FAILED(xrEnumerateApiLayerProperties(0, &LayerPropertyCount, nullptr)))
-	{
-		// As per EnumerateExtensions - a failure here means no runtime installed.
-		return false;
-	}
-
-	if (!LayerPropertyCount)
-	{
-		// It's still legit if we have no layers, so early out here (and return success) if so.
-		return true;
-	}
-
-	TArray<XrApiLayerProperties> LayerProperties;
-	LayerProperties.SetNum(LayerPropertyCount);
-	for (auto& Prop : LayerProperties)
-	{
-		Prop = XrApiLayerProperties{ XR_TYPE_API_LAYER_PROPERTIES };
-	}
-
-	if (XR_ENSURE(xrEnumerateApiLayerProperties(LayerPropertyCount, &LayerPropertyCount, LayerProperties.GetData())))
-	{
-		for (const auto& Prop : LayerProperties)
-		{
-			AvailableLayers.Add(Prop.layerName);
-		}
-		return true;
-	}
-
-	return false;
-}
-
-struct AnsiKeyFunc : BaseKeyFuncs<const ANSICHAR*, const ANSICHAR*, false>
-{
-	typedef typename TTypeTraits<const ANSICHAR*>::ConstPointerType KeyInitType;
-	typedef typename TCallTraits<const ANSICHAR*>::ParamType ElementInitType;
-
-	/**
-	 * @return The key used to index the given element.
-	 */
-	static FORCEINLINE KeyInitType GetSetKey(ElementInitType Element)
-	{
-		return Element;
-	}
-
-	/**
-	 * @return True if the keys match.
-	 */
-	static FORCEINLINE bool Matches(KeyInitType A, KeyInitType B)
-	{
-		return FCStringAnsi::Strcmp(A, B) == 0;
-	}
-
-	/** Calculates a hash index for a key. */
-	static FORCEINLINE uint32 GetKeyHash(KeyInitType Key)
-	{
-		return GetTypeHash(Key);
-	}
-};
-
-bool FOpenXRHMDPlugin::InitRenderBridge()
-{
-	// Get all extension plugins
-	TSet<const ANSICHAR*, AnsiKeyFunc> ExtensionSet;
-	TArray<IOpenXRExtensionPlugin*> ExtModules = IModularFeatures::Get().GetModularFeatureImplementations<IOpenXRExtensionPlugin>(IOpenXRExtensionPlugin::GetModularFeatureName());
-
-	// Query all extension plugins to see if we need to use a custom render bridge
-	PFN_xrGetInstanceProcAddr GetProcAddr = nullptr;
-	for (IOpenXRExtensionPlugin* Plugin : ExtModules)
-	{
-		// We are taking ownership of the CustomRenderBridge instance here.
-		TRefCountPtr<FOpenXRRenderBridge> CustomRenderBridge = Plugin->GetCustomRenderBridge(Instance, System);
-		if (CustomRenderBridge)
-		{
-			// We pick the first
-			RenderBridge = CustomRenderBridge;
-			return true;
-		}
-	}
-
-
-	FString RHIString = FApp::GetGraphicsRHI();
-	if (RHIString.IsEmpty())
-	{
-		return false;
-	}
-
-	if (!InitInstanceAndSystem())
-	{
-		return false;
-	}
-
-#ifdef XR_USE_GRAPHICS_API_D3D11
-	if (RHIString == TEXT("DirectX 11") && IsExtensionEnabled(XR_KHR_D3D11_ENABLE_EXTENSION_NAME))
-	{
-		RenderBridge = CreateRenderBridge_D3D11(Instance, System);
-	}
-	else
-#endif
-#ifdef XR_USE_GRAPHICS_API_D3D12
-	if (RHIString == TEXT("DirectX 12") && IsExtensionEnabled(XR_KHR_D3D12_ENABLE_EXTENSION_NAME))
-	{
-		RenderBridge = CreateRenderBridge_D3D12(Instance, System);
-	}
-	else
-#endif
-#ifdef XR_USE_GRAPHICS_API_OPENGL
-	if (RHIString == TEXT("OpenGL") && IsExtensionEnabled(XR_KHR_OPENGL_ENABLE_EXTENSION_NAME))
-	{
-		RenderBridge = CreateRenderBridge_OpenGL(Instance, System);
-	}
-	else
-#endif
-#ifdef XR_USE_GRAPHICS_API_VULKAN
-	if (RHIString == TEXT("Vulkan") && IsExtensionEnabled(XR_KHR_VULKAN_ENABLE_EXTENSION_NAME))
-	{
-		RenderBridge = CreateRenderBridge_Vulkan(Instance, System);
-	}
-	else
-#endif
-	{
-		UE_LOG(LogHMD, Warning, TEXT("%s is not currently supported by the OpenXR runtime"), *RHIString);
-		return false;
-	}
-	return true;
-}
-
-PFN_xrGetInstanceProcAddr FOpenXRHMDPlugin::GetDefaultLoader()
-{
-#if PLATFORM_WINDOWS
-#if !PLATFORM_CPU_X86_FAMILY
-#error Windows platform does not currently support this CPU family. A OpenXR loader binary for this CPU family is needed.
-#endif
-
-#if PLATFORM_64BITS
-	FString BinariesPath = FPaths::EngineDir() / FString(TEXT("Binaries/ThirdParty/OpenXR/win64"));
-#else
-	FString BinariesPath = FPaths::EngineDir() / FString(TEXT("Binaries/ThirdParty/OpenXR/win32"));
-#endif
-
-	FString LoaderName = "openxr_loader.dll";
-	FPlatformProcess::PushDllDirectory(*BinariesPath);
-	LoaderHandle = FPlatformProcess::GetDllHandle(*LoaderName);
-	FPlatformProcess::PopDllDirectory(*BinariesPath);
-#elif PLATFORM_LINUX
-	LoaderHandle = FPlatformProcess::GetDllHandle(TEXT("/usr/lib/libopenxr_loader.so"));
-	if (!LoaderHandle)
-	{
-		LoaderHandle = FPlatformProcess::GetDllHandle(TEXT("/usr/lib/x86_64-linux-gnu/libopenxr_loader.so"));
-		if (!LoaderHandle)
-		{
-			LoaderHandle = FPlatformProcess::GetDllHandle(TEXT("/usr/local/lib/libopenxr_loader.so"));
-		}
-	}
-#elif PLATFORM_HOLOLENS
-#ifndef PLATFORM_64BITS
-#error HoloLens platform does not currently support 32-bit. 32-bit OpenXR loader binaries are needed.
-#endif
-
-#if PLATFORM_CPU_ARM_FAMILY
-	FString BinariesPath = FPaths::EngineDir() / FString(TEXT("Binaries/ThirdParty/OpenXR/hololens/arm64"));
-#elif PLATFORM_CPU_X86_FAMILY
-	FString BinariesPath = FPaths::EngineDir() / FString(TEXT("Binaries/ThirdParty/OpenXR/hololens/x64"));
-#else
-#error Unsupported CPU family for the HoloLens platform.
-#endif
-
-	LoaderHandle = FPlatformProcess::GetDllHandle(*(BinariesPath / "openxr_loader.dll"));
-#endif
-
-	if (!LoaderHandle)
-	{
-		UE_LOG(LogHMD, Log, TEXT("Failed to load openxr_loader.dll"));
-		return nullptr;
-	}
-	return (PFN_xrGetInstanceProcAddr)FPlatformProcess::GetDllExport(LoaderHandle, TEXT("xrGetInstanceProcAddr"));
-}
-
-bool FOpenXRHMDPlugin::EnableExtensions(const TArray<const ANSICHAR*>& RequiredExtensions, const TArray<const ANSICHAR*>& OptionalExtensions, TArray<const ANSICHAR*>& OutExtensions)
-{
-	// Query required extensions and check if they're all available
-	bool ExtensionMissing = false;
-	for (const ANSICHAR* Ext : RequiredExtensions)
-	{
-		if (AvailableExtensions.Contains(Ext))
-		{
-			UE_LOG(LogHMD, Verbose, TEXT("Required extension %S enabled"), Ext);
-		}
-		else
-		{
-			UE_LOG(LogHMD, Warning, TEXT("Required extension %S is not available"), Ext);
-			ExtensionMissing = true;
-		}
-	}
-
-	// If any required extensions are missing then we ignore the plugin
-	if (ExtensionMissing)
-	{
-		return false;
-	}
-
-	// All required extensions are supported we can safely add them to our set and give the plugin callbacks
-	OutExtensions.Append(RequiredExtensions);
-
-	// Add all supported optional extensions to the set
-	for (const ANSICHAR* Ext : OptionalExtensions)
-	{
-		if (AvailableExtensions.Contains(Ext))
-		{
-			UE_LOG(LogHMD, Verbose, TEXT("Optional extension %S enabled"), Ext);
-			OutExtensions.Add(Ext);
-		}
-		else
-		{
-			UE_LOG(LogHMD, Log, TEXT("Optional extension %S is not available"), Ext);
-		}
-	}
-
-	return true;
-}
-
-bool FOpenXRHMDPlugin::GetRequiredExtensions(TArray<const ANSICHAR*>& OutExtensions)
-{
-#if PLATFORM_ANDROID
-	OutExtensions.Add(XR_KHR_ANDROID_CREATE_INSTANCE_EXTENSION_NAME);
-#endif
-	return true;
-}
-
-bool FOpenXRHMDPlugin::GetOptionalExtensions(TArray<const ANSICHAR*>& OutExtensions)
-{
-#ifdef XR_USE_GRAPHICS_API_D3D11
-	OutExtensions.Add(XR_KHR_D3D11_ENABLE_EXTENSION_NAME);
-#endif
-#ifdef XR_USE_GRAPHICS_API_D3D12
-	OutExtensions.Add(XR_KHR_D3D12_ENABLE_EXTENSION_NAME);
-#endif
-#ifdef XR_USE_GRAPHICS_API_OPENGL
-	OutExtensions.Add(XR_KHR_OPENGL_ENABLE_EXTENSION_NAME);
-#endif
-#ifdef XR_USE_GRAPHICS_API_VULKAN
-	OutExtensions.Add(XR_KHR_VULKAN_ENABLE_EXTENSION_NAME);
-#endif
-	OutExtensions.Add(XR_KHR_COMPOSITION_LAYER_DEPTH_EXTENSION_NAME);
-	OutExtensions.Add(XR_VARJO_QUAD_VIEWS_EXTENSION_NAME);
-	OutExtensions.Add(XR_KHR_VISIBILITY_MASK_EXTENSION_NAME);
-	return true;
-}
-
-bool FOpenXRHMDPlugin::InitInstanceAndSystem()
-{
-	if (!Instance && !InitInstance())
-	{
-		return false;
-	}
-
-	if (!System && !InitSystem())
-	{
-		return false;
-	}
-
-	return true;
-}
-
-bool FOpenXRHMDPlugin::InitInstance()
-{
-	// This should only ever be called if we don't already have an instance.
-	check(!Instance);
-
-	// Get all extension plugins
-	TSet<const ANSICHAR*, AnsiKeyFunc> ExtensionSet;
-	TArray<IOpenXRExtensionPlugin*> ExtModules = IModularFeatures::Get().GetModularFeatureImplementations<IOpenXRExtensionPlugin>(IOpenXRExtensionPlugin::GetModularFeatureName());
-
-	// Query all extension plugins to see if we need to use a custom loader
-	PFN_xrGetInstanceProcAddr GetProcAddr = nullptr;
-	for (IOpenXRExtensionPlugin* Plugin : ExtModules)
-	{
-		if (Plugin->GetCustomLoader(&GetProcAddr))
-		{
-			// We pick the first loader we can find
-			break;
-		}
-
-		// Clear it again just to ensure the failed call didn't leave the pointer set
-		GetProcAddr = nullptr;
-	}
-
-	if (!GetProcAddr)
-	{
-		GetProcAddr = GetDefaultLoader();
-	}
-
-	if (!PreInitOpenXRCore(GetProcAddr))
-	{
-		UE_LOG(LogHMD, Log, TEXT("Failed to initialize core functions. Please check that you have a valid OpenXR runtime installed."));
-		return false;
-	}
-
-	if (!EnumerateExtensions())
-	{
-		UE_LOG(LogHMD, Log, TEXT("Failed to enumerate extensions. Please check that you have a valid OpenXR runtime installed."));
-		return false;
-	}
-
-	if (!EnumerateLayers())
-	{
-		UE_LOG(LogHMD, Log, TEXT("Failed to enumerate API layers. Please check that you have a valid OpenXR runtime installed."));
-		return false;
-	}
-
-	// Enable any required and optional extensions that are not plugin specific (usually platform support extensions)
-	{
-		TArray<const ANSICHAR*> RequiredExtensions, OptionalExtensions, Extensions;
-		// Query required extensions
-		RequiredExtensions.Empty();
-		if (!GetRequiredExtensions(RequiredExtensions))
-		{
-			UE_LOG(LogHMD, Error, TEXT("Could not get required OpenXR extensions."));
-			return false;
-		}
-
-		// Query optional extensions
-		OptionalExtensions.Empty();
-		if (!GetOptionalExtensions(OptionalExtensions))
-		{
-			UE_LOG(LogHMD, Error, TEXT("Could not get optional OpenXR extensions."));
-			return false;
-		}
-
-		if (!EnableExtensions(RequiredExtensions, OptionalExtensions, Extensions))
-		{
-			UE_LOG(LogHMD, Error, TEXT("Could not enable all required OpenXR extensions."));
-			return false;
-		}
-		ExtensionSet.Append(Extensions);
-	}
-
-	if (AvailableExtensions.Contains(XR_EPIC_VIEW_CONFIGURATION_FOV_EXTENSION_NAME))
-	{
-		ExtensionSet.Add(XR_EPIC_VIEW_CONFIGURATION_FOV_EXTENSION_NAME);
-	}
-
-	for (IOpenXRExtensionPlugin* Plugin : ExtModules)
-	{
-		TArray<const ANSICHAR*> RequiredExtensions, OptionalExtensions, Extensions;
-
-		// Query required extensions
-		RequiredExtensions.Empty();
-		if (!Plugin->GetRequiredExtensions(RequiredExtensions))
-		{
-			// Ignore the plugin if the query fails
-			continue;
-		}
-
-		// Query optional extensions
-		OptionalExtensions.Empty();
-		if (!Plugin->GetOptionalExtensions(OptionalExtensions))
-		{
-			// Ignore the plugin if the query fails
-			continue;
-		}
-
-		if (!EnableExtensions(RequiredExtensions, OptionalExtensions, Extensions))
-		{
-			// Ignore the plugin if the required extension could not be enabled
-			FString ModuleName = Plugin->GetDisplayName();
-			UE_LOG(LogHMD, Log, TEXT("Could not enable all required OpenXR extensions for %s on current system. This plugin will be loaded but ignored, but will be enabled on a target platform that supports the required extension."), *ModuleName);
-			continue;
-		}
-		ExtensionSet.Append(Extensions);
-		ExtensionPlugins.Add(Plugin);
-	}
-
-	if (auto ARModule = FModuleManager::LoadModulePtr<IOpenXRARModule>("OpenXRAR"))
-	{
-		TArray<const ANSICHAR*> ARExtensionSet;
-		ARModule->GetExtensions(ARExtensionSet);
-		ExtensionSet.Append(ARExtensionSet);
-	}
-
-	EnabledExtensions.Reset();
-	for (const ANSICHAR* Ext : ExtensionSet)
-	{
-		EnabledExtensions.Add(Ext);
-	}
-
-	// Enable layers, if specified by CVar.
-	// Note: For the validation layer to work on Windows (as of latest OpenXR runtime, August 2019), the following are required:
-	//   1. Download and build the OpenXR SDK from [REDACTED_PROFILE_URL]/OpenXR-SDK-Source (follow instructions at [REDACTED_PROFILE_URL]/OpenXR-SDK-Source/blob/master/BUILDING.md)
-	//	 2. Add a registry key under HKEY_LOCAL_MACHINE\SOFTWARE\Khronos\OpenXR\1\ApiLayers\Explicit, containing the path to the manifest file
-	//      (e.g. C:\OpenXR-SDK-Source-master\build\win64\src\api_layers\XrApiLayer_core_validation.json) <-- this file is downloaded as part of the SDK source, above
-	//   3. Copy the DLL from the build target at, for example, C:\OpenXR-SDK-Source-master\build\win64\src\api_layers\XrApiLayer_core_validation.dll to
-	//      somewhere in your system path (e.g. c:\windows\system32); the OpenXR loader currently doesn't use the path the json file is in (this is a bug)
-
-	const bool bEnableOpenXRValidationLayer = CVarEnableOpenXRValidationLayer.GetValueOnAnyThread() != 0;
-	TArray<const char*> Layers;
-	if (bEnableOpenXRValidationLayer && AvailableLayers.Contains("XR_APILAYER_LUNARG_core_validation"))
-	{
-		Layers.Add("XR_APILAYER_LUNARG_core_validation");
-	}
-
-	// Engine registration can be disabled via console var.
-	auto* CVarDisableEngineAndAppRegistration = IConsoleManager::Get().FindTConsoleVariableDataInt(TEXT("r.DisableEngineAndAppRegistration"));
-	bool bDisableEngineRegistration = (CVarDisableEngineAndAppRegistration && CVarDisableEngineAndAppRegistration->GetValueOnAnyThread() != 0);
-
-	FText ProjectName = FText();
-	GConfig->GetText(TEXT("/Script/EngineSettings.GeneralProjectSettings"), TEXT("ProjectName"), ProjectName, GGameIni);
-
-	FText ProjectVersion = FText();
-	GConfig->GetText(TEXT("/Script/EngineSettings.GeneralProjectSettings"), TEXT("ProjectVersion"), ProjectVersion, GGameIni);
-
-	// EngineName will be of the form "UnrealEngine4.21", with the minor version ("21" in this example)
-	// updated with every quarterly release
-	FString EngineName = bDisableEngineRegistration ? FString("") : FApp::GetEpicProductIdentifier() + FEngineVersion::Current().ToString(EVersionComponent::Minor);
-	FString AppName = bDisableEngineRegistration ? TEXT("") : ProjectName.ToString() + ProjectVersion.ToString();
-
-	XrInstanceCreateInfo Info;
-	Info.type = XR_TYPE_INSTANCE_CREATE_INFO;
-	Info.next = nullptr;
-	Info.createFlags = 0;
-	FTCHARToUTF8_Convert::Convert(Info.applicationInfo.applicationName, XR_MAX_APPLICATION_NAME_SIZE, *AppName, AppName.Len() + 1);
-	Info.applicationInfo.applicationVersion = static_cast<uint32>(BuildSettings::GetCurrentChangelist()) | (BuildSettings::IsLicenseeVersion() ? 0x80000000 : 0);
-	FTCHARToUTF8_Convert::Convert(Info.applicationInfo.engineName, XR_MAX_ENGINE_NAME_SIZE, *EngineName, EngineName.Len() + 1);
-	Info.applicationInfo.engineVersion = (uint32)(FEngineVersion::Current().GetMinor() << 16 | FEngineVersion::Current().GetPatch());
-	Info.applicationInfo.apiVersion = XR_CURRENT_API_VERSION;
-
-	Info.enabledApiLayerCount = Layers.Num();
-	Info.enabledApiLayerNames = Layers.GetData();
-
-	Info.enabledExtensionCount = EnabledExtensions.Num();
-	Info.enabledExtensionNames = EnabledExtensions.GetData();
-
-#if PLATFORM_ANDROID
-	XrInstanceCreateInfoAndroidKHR InstanceCreateInfoAndroid;
-	InstanceCreateInfoAndroid.type = XR_TYPE_INSTANCE_CREATE_INFO_ANDROID_KHR;
-	InstanceCreateInfoAndroid.next = nullptr;
-	InstanceCreateInfoAndroid.applicationVM = GNativeAndroidApp->activity->vm;
-	InstanceCreateInfoAndroid.applicationActivity = GNativeAndroidApp->activity->clazz;
-	Info.next = &InstanceCreateInfoAndroid;
-#endif // PLATFORM_ANDROID
-
-	for (IOpenXRExtensionPlugin* Module : ExtensionPlugins)
-	{
-		Info.next = Module->OnCreateInstance(this, Info.next);
-	}
-
-	XrResult Result = xrCreateInstance(&Info, &Instance);
-	if (XR_FAILED(Result))
-	{
-		UE_LOG(LogHMD, Log, TEXT("Failed to create an OpenXR instance, result is %s. Please check if you have an OpenXR runtime installed. The following extensions were enabled:"), OpenXRResultToString(Result));
-		for (const char* Extension : EnabledExtensions)
-		{
-			UE_LOG(LogHMD, Log, TEXT("- %S"), Extension);
-		}
-		return false;
-	}
-
-	if (!InitOpenXRCore(Instance))
-	{
-		UE_LOG(LogHMD, Log, TEXT("Failed to initialize core functions. Please check that you have a valid OpenXR runtime installed."));
-		return false;
-	}
-
-	for (IOpenXRExtensionPlugin* Module : ExtensionPlugins)
-	{
-		Module->PostCreateInstance(Instance);
-	}
-
-	return true;
-}
-
-bool FOpenXRHMDPlugin::InitSystem()
-{
-	XrSystemGetInfo SystemInfo;
-	SystemInfo.type = XR_TYPE_SYSTEM_GET_INFO;
-	SystemInfo.next = nullptr;
-	SystemInfo.formFactor = XR_FORM_FACTOR_HEAD_MOUNTED_DISPLAY;
-	for (IOpenXRExtensionPlugin* Module : ExtensionPlugins)
-	{
-		SystemInfo.next = Module->OnGetSystem(Instance, SystemInfo.next);
-	}
-
-	XrResult Result = xrGetSystem(Instance, &SystemInfo, &System);
-	if (XR_FAILED(Result))
-	{
-		UE_LOG(LogHMD, Log, TEXT("Failed to get an OpenXR system, result is %s. Please check that your runtime supports VR headsets."), OpenXRResultToString(Result));
-		return false;
-	}
-
-	for (IOpenXRExtensionPlugin* Module : ExtensionPlugins)
-	{
-		Module->PostGetSystem(Instance, System);
-	}
-
-	return true;
-}
-
 //---------------------------------------------------
 // OpenXRHMD IHeadMountedDisplay Implementation
 //---------------------------------------------------
diff --git a/Engine/Plugins/Runtime/OpenXR/Source/OpenXRHMD/Private/OpenXRHMDModule.cpp b/Engine/Plugins/Runtime/OpenXR/Source/OpenXRHMD/Private/OpenXRHMDModule.cpp
new file mode 100644
--- /dev/null
+++ b/Engine/Plugins/Runtime/OpenXR/Source/OpenXRHMD/Private/OpenXRHMDModule.cpp
@@ -0,0 +1,632 @@
+// Copyright Epic Games, Inc. All Rights Reserved.
+
+#include "OpenXRHMDModule.h"
+#include "OpenXRHMD.h"
+#include "OpenXRHMD_RenderBridge.h"
+#include "OpenXRCore.h"
+#include "IOpenXRExtensionPlugin.h"
+#include "IOpenXRARModule.h"
+#include "BuildSettings.h"
+
+#if PLATFORM_ANDROID
+#include <android_native_app_glue.h>
+extern struct android_app* GNativeAndroidApp;
+#endif
+
+static TAutoConsoleVariable<int32> CVarEnableOpenXRValidationLayer(
+	TEXT("xr.EnableOpenXRValidationLayer"),
+	0,
+	TEXT("If true, enables the OpenXR validation layer, which will provide extended validation of\nOpenXR API calls. This should only be used for debugging purposes.\n")
+	TEXT("Changes will only take effect in new game/editor instances - can't be changed at runtime.\n"),
+	ECVF_Default);		// @todo: Should we specify ECVF_Cheat here so this doesn't show up in release builds?
+
+//---------------------------------------------------
+// OpenXRHMD Plugin Implementation
+//---------------------------------------------------
+
+IMPLEMENT_MODULE( FOpenXRHMDModule, OpenXRHMD )
+
+TSharedPtr< class IXRTrackingSystem, ESPMode::ThreadSafe > FOpenXRHMDModule::CreateTrackingSystem()
+{
+	if (!RenderBridge)
+	{
+		if (!InitRenderBridge())
+		{
+			return nullptr;
+		}
+	}
+	auto ARModule = FModuleManager::LoadModulePtr<IOpenXRARModule>("OpenXRAR");
+	auto ARSystem = ARModule->CreateARSystem();
+
+	auto OpenXRHMD = FSceneViewExtensions::NewExtension<FOpenXRHMD>(Instance, System, RenderBridge, EnabledExtensions, ExtensionPlugins, ARSystem);
+	if (OpenXRHMD->IsInitialized())
+	{
+		ARModule->SetTrackingSystem(OpenXRHMD);
+		OpenXRHMD->GetARCompositionComponent()->InitializeARSystem();
+		return OpenXRHMD;
+	}
+
+	return nullptr;
+}
+
+void FOpenXRHMDModule::ShutdownModule()
+{
+	if (Instance)
+	{
+		XR_ENSURE(xrDestroyInstance(Instance));
+	}
+
+	if (LoaderHandle)
+	{
+		FPlatformProcess::FreeDllHandle(LoaderHandle);
+		LoaderHandle = nullptr;
+	}
+}
+
+uint64 FOpenXRHMDModule::GetGraphicsAdapterLuid()
+{
+	if (!RenderBridge)
+	{
+		if (!InitRenderBridge())
+		{
+			return 0;
+		}
+	}
+	return RenderBridge->GetGraphicsAdapterLuid();
+}
+
+TSharedPtr< IHeadMountedDisplayVulkanExtensions, ESPMode::ThreadSafe > FOpenXRHMDModule::GetVulkanExtensions()
+{
+#ifdef XR_USE_GRAPHICS_API_VULKAN
+	if (InitInstanceAndSystem() && IsExtensionEnabled(XR_KHR_VULKAN_ENABLE_EXTENSION_NAME))
+	{
+		if (!VulkanExtensions.IsValid())
+		{
+			VulkanExtensions = MakeShareable(new FOpenXRHMD::FVulkanExtensions(Instance, System));
+		}
+		return VulkanExtensions;
+	}
+#endif//XR_USE_GRAPHICS_API_VULKAN
+	return nullptr;
+}
+
+bool FOpenXRHMDModule::IsStandaloneStereoOnlyDevice()
+{
+	if (InitInstanceAndSystem())
+	{
+		for (IOpenXRExtensionPlugin* Module : ExtensionPlugins)
+		{
+			if (Module->IsStandaloneStereoOnlyDevice())
+			{
+				return true;
+			}
+		}
+	}
+	return false;
+}
+
+bool FOpenXRHMDModule::EnumerateExtensions()
+{
+	uint32_t ExtensionsCount = 0;
+	if (XR_FAILED(xrEnumerateInstanceExtensionProperties(nullptr, 0, &ExtensionsCount, nullptr)))
+	{
+		// If it fails this early that means there's no runtime installed
+		return false;
+	}
+
+	TArray<XrExtensionProperties> Properties;
+	Properties.SetNum(ExtensionsCount);
+	for (auto& Prop : Properties)
+	{
+		Prop = XrExtensionProperties{ XR_TYPE_EXTENSION_PROPERTIES };
+	}
+
+	if (XR_ENSURE(xrEnumerateInstanceExtensionProperties(nullptr, ExtensionsCount, &ExtensionsCount, Properties.GetData())))
+	{
+		for (const XrExtensionProperties& Prop : Properties)
+		{
+			AvailableExtensions.Add(Prop.extensionName);
+		}
+		return true;
+	}
+	return false;
+}
+
+bool FOpenXRHMDModule::EnumerateLayers()
+{
+	uint32 LayerPropertyCount = 0;
+	if (XR_FAILED(xrEnumerateApiLayerProperties(0, &LayerPropertyCount, nullptr)))
+	{
+		// As per EnumerateExtensions - a failure here means no runtime installed.
+		return false;
+	}
+
+	if (!LayerPropertyCount)
+	{
+		// It's still legit if we have no layers, so early out here (and return success) if so.
+		return true;
+	}
+
+	TArray<XrApiLayerProperties> LayerProperties;
+	LayerProperties.SetNum(LayerPropertyCount);
+	for (auto& Prop : LayerProperties)
+	{
+		Prop = XrApiLayerProperties{ XR_TYPE_API_LAYER_PROPERTIES };
+	}
+
+	if (XR_ENSURE(xrEnumerateApiLayerProperties(LayerPropertyCount, &LayerPropertyCount, LayerProperties.GetData())))
+	{
+		for (const auto& Prop : LayerProperties)
+		{
+			AvailableLayers.Add(Prop.layerName);
+		}
+		return true;
+	}
+
+	return false;
+}
+
+struct AnsiKeyFunc : BaseKeyFuncs<const ANSICHAR*, const ANSICHAR*, false>
+{
+	typedef typename TTypeTraits<const ANSICHAR*>::ConstPointerType KeyInitType;
+	typedef typename TCallTraits<const ANSICHAR*>::ParamType ElementInitType;
+
+	/**
+	 * @return The key used to index the given element.
+	 */
+	static FORCEINLINE KeyInitType GetSetKey(ElementInitType Element)
+	{
+		return Element;
+	}
+
+	/**
+	 * @return True if the keys match.
+	 */
+	static FORCEINLINE bool Matches(KeyInitType A, KeyInitType B)
+	{
+		return FCStringAnsi::Strcmp(A, B) == 0;
+	}
+
+	/** Calculates a hash index for a key. */
+	static FORCEINLINE uint32 GetKeyHash(KeyInitType Key)
+	{
+		return GetTypeHash(Key);
+	}
+};
+
+bool FOpenXRHMDModule::InitRenderBridge()
+{
+	// Get all extension plugins
+	TSet<const ANSICHAR*, AnsiKeyFunc> ExtensionSet;
+	TArray<IOpenXRExtensionPlugin*> ExtModules = IModularFeatures::Get().GetModularFeatureImplementations<IOpenXRExtensionPlugin>(IOpenXRExtensionPlugin::GetModularFeatureName());
+
+	// Query all extension plugins to see if we need to use a custom render bridge
+	PFN_xrGetInstanceProcAddr GetProcAddr = nullptr;
+	for (IOpenXRExtensionPlugin* Plugin : ExtModules)
+	{
+		// We are taking ownership of the CustomRenderBridge instance here.
+		TRefCountPtr<FOpenXRRenderBridge> CustomRenderBridge = Plugin->GetCustomRenderBridge(Instance, System);
+		if (CustomRenderBridge)
+		{
+			// We pick the first
+			RenderBridge = CustomRenderBridge;
+			return true;
+		}
+	}
+
+
+	FString RHIString = FApp::GetGraphicsRHI();
+	if (RHIString.IsEmpty())
+	{
+		return false;
+	}
+
+	if (!InitInstanceAndSystem())
+	{
+		return false;
+	}
+
+#ifdef XR_USE_GRAPHICS_API_D3D11
+	if (RHIString == TEXT("DirectX 11") && IsExtensionEnabled(XR_KHR_D3D11_ENABLE_EXTENSION_NAME))
+	{
+		RenderBridge = CreateRenderBridge_D3D11(Instance, System);
+	}
+	else
+#endif
+#ifdef XR_USE_GRAPHICS_API_D3D12
+	if (RHIString == TEXT("DirectX 12") && IsExtensionEnabled(XR_KHR_D3D12_ENABLE_EXTENSION_NAME))
+	{
+		RenderBridge = CreateRenderBridge_D3D12(Instance, System);
+	}
+	else
+#endif
+#ifdef XR_USE_GRAPHICS_API_OPENGL
+	if (RHIString == TEXT("OpenGL") && IsExtensionEnabled(XR_KHR_OPENGL_ENABLE_EXTENSION_NAME))
+	{
+		RenderBridge = CreateRenderBridge_OpenGL(Instance, System);
+	}
+	else
+#endif
+#ifdef XR_USE_GRAPHICS_API_VULKAN
+	if (RHIString == TEXT("Vulkan") && IsExtensionEnabled(XR_KHR_VULKAN_ENABLE_EXTENSION_NAME))
+	{
+		RenderBridge = CreateRenderBridge_Vulkan(Instance, System);
+	}
+	else
+#endif
+	{
+		UE_LOG(LogHMD, Warning, TEXT("%s is not currently supported by the OpenXR runtime"), *RHIString);
+		return false;
+	}
+	return true;
+}
+
+PFN_xrGetInstanceProcAddr FOpenXRHMDModule::GetDefaultLoader()
+{
+#if PLATFORM_WINDOWS
+#if !PLATFORM_CPU_X86_FAMILY
+#error Windows platform does not currently support this CPU family. A OpenXR loader binary for this CPU family is needed.
+#endif
+
+#if PLATFORM_64BITS
+	FString BinariesPath = FPaths::EngineDir() / FString(TEXT("Binaries/ThirdParty/OpenXR/win64"));
+#else
+	FString BinariesPath = FPaths::EngineDir() / FString(TEXT("Binaries/ThirdParty/OpenXR/win32"));
+#endif
+
+	FString LoaderName = "openxr_loader.dll";
+	FPlatformProcess::PushDllDirectory(*BinariesPath);
+	LoaderHandle = FPlatformProcess::GetDllHandle(*LoaderName);
+	FPlatformProcess::PopDllDirectory(*BinariesPath);
+#elif PLATFORM_LINUX
+	LoaderHandle = FPlatformProcess::GetDllHandle(TEXT("/usr/lib/libopenxr_loader.so"));
+	if (!LoaderHandle)
+	{
+		LoaderHandle = FPlatformProcess::GetDllHandle(TEXT("/usr/lib/x86_64-linux-gnu/libopenxr_loader.so"));
+		if (!LoaderHandle)
+		{
+			LoaderHandle = FPlatformProcess::GetDllHandle(TEXT("/usr/local/lib/libopenxr_loader.so"));
+		}
+	}
+#elif PLATFORM_HOLOLENS
+#ifndef PLATFORM_64BITS
+#error HoloLens platform does not currently support 32-bit. 32-bit OpenXR loader binaries are needed.
+#endif
+
+#if PLATFORM_CPU_ARM_FAMILY
+	FString BinariesPath = FPaths::EngineDir() / FString(TEXT("Binaries/ThirdParty/OpenXR/hololens/arm64"));
+#elif PLATFORM_CPU_X86_FAMILY
+	FString BinariesPath = FPaths::EngineDir() / FString(TEXT("Binaries/ThirdParty/OpenXR/hololens/x64"));
+#else
+#error Unsupported CPU family for the HoloLens platform.
+#endif
+
+	LoaderHandle = FPlatformProcess::GetDllHandle(*(BinariesPath / "openxr_loader.dll"));
+#endif
+
+	if (!LoaderHandle)
+	{
+		UE_LOG(LogHMD, Log, TEXT("Failed to load openxr_loader.dll"));
+		return nullptr;
+	}
+	return (PFN_xrGetInstanceProcAddr)FPlatformProcess::GetDllExport(LoaderHandle, TEXT("xrGetInstanceProcAddr"));
+}
+
+bool FOpenXRHMDModule::EnableExtensions(const TArray<const ANSICHAR*>& RequiredExtensions, const TArray<const ANSICHAR*>& OptionalExtensions, TArray<const ANSICHAR*>& OutExtensions)
+{
+	// Query required extensions and check if they're all available
+	bool ExtensionMissing = false;
+	for (const ANSICHAR* Ext : RequiredExtensions)
+	{
+		if (AvailableExtensions.Contains(Ext))
+		{
+			UE_LOG(LogHMD, Verbose, TEXT("Required extension %S enabled"), Ext);
+		}
+		else
+		{
+			UE_LOG(LogHMD, Warning, TEXT("Required extension %S is not available"), Ext);
+			ExtensionMissing = true;
+		}
+	}
+
+	// If any required extensions are missing then we ignore the plugin
+	if (ExtensionMissing)
+	{
+		return false;
+	}
+
+	// All required extensions are supported we can safely add them to our set and give the plugin callbacks
+	OutExtensions.Append(RequiredExtensions);
+
+	// Add all supported optional extensions to the set
+	for (const ANSICHAR* Ext : OptionalExtensions)
+	{
+		if (AvailableExtensions.Contains(Ext))
+		{
+			UE_LOG(LogHMD, Verbose, TEXT("Optional extension %S enabled"), Ext);
+			OutExtensions.Add(Ext);
+		}
+		else
+		{
+			UE_LOG(LogHMD, Log, TEXT("Optional extension %S is not available"), Ext);
+		}
+	}
+
+	return true;
+}
+
+bool FOpenXRHMDModule::GetRequiredExtensions(TArray<const ANSICHAR*>& OutExtensions)
+{
+#if PLATFORM_ANDROID
+	OutExtensions.Add(XR_KHR_ANDROID_CREATE_INSTANCE_EXTENSION_NAME);
+#endif
+	return true;
+}
+
+bool FOpenXRHMDModule::GetOptionalExtensions(TArray<const ANSICHAR*>& OutExtensions)
+{
+#ifdef XR_USE_GRAPHICS_API_D3D11
+	OutExtensions.Add(XR_KHR_D3D11_ENABLE_EXTENSION_NAME);
+#endif
+#ifdef XR_USE_GRAPHICS_API_D3D12
+	OutExtensions.Add(XR_KHR_D3D12_ENABLE_EXTENSION_NAME);
+#endif
+#ifdef XR_USE_GRAPHICS_API_OPENGL
+	OutExtensions.Add(XR_KHR_OPENGL_ENABLE_EXTENSION_NAME);
+#endif
+#ifdef XR_USE_GRAPHICS_API_VULKAN
+	OutExtensions.Add(XR_KHR_VULKAN_ENABLE_EXTENSION_NAME);
+#endif
+	OutExtensions.Add(XR_KHR_COMPOSITION_LAYER_DEPTH_EXTENSION_NAME);
+	OutExtensions.Add(XR_VARJO_QUAD_VIEWS_EXTENSION_NAME);
+	OutExtensions.Add(XR_KHR_VISIBILITY_MASK_EXTENSION_NAME);
+	return true;
+}
+
+bool FOpenXRHMDModule::InitInstanceAndSystem()
+{
+	if (!Instance && !InitInstance())
+	{
+		return false;
+	}
+
+	if (!System && !InitSystem())
+	{
+		return false;
+	}
+
+	return true;
+}
+
+bool FOpenXRHMDModule::InitInstance()
+{
+	// This should only ever be called if we don't already have an instance.
+	check(!Instance);
+
+	// Get all extension plugins
+	TSet<const ANSICHAR*, AnsiKeyFunc> ExtensionSet;
+	TArray<IOpenXRExtensionPlugin*> ExtModules = IModularFeatures::Get().GetModularFeatureImplementations<IOpenXRExtensionPlugin>(IOpenXRExtensionPlugin::GetModularFeatureName());
+
+	// Query all extension plugins to see if we need to use a custom loader
+	PFN_xrGetInstanceProcAddr GetProcAddr = nullptr;
+	for (IOpenXRExtensionPlugin* Plugin : ExtModules)
+	{
+		if (Plugin->GetCustomLoader(&GetProcAddr))
+		{
+			// We pick the first loader we can find
+			break;
+		}
+
+		// Clear it again just to ensure the failed call didn't leave the pointer set
+		GetProcAddr = nullptr;
+	}
+
+	if (!GetProcAddr)
+	{
+		GetProcAddr = GetDefaultLoader();
+	}
+
+	if (!PreInitOpenXRCore(GetProcAddr))
+	{
+		UE_LOG(LogHMD, Log, TEXT("Failed to initialize core functions. Please check that you have a valid OpenXR runtime installed."));
+		return false;
+	}
+
+	if (!EnumerateExtensions())
+	{
+		UE_LOG(LogHMD, Log, TEXT("Failed to enumerate extensions. Please check that you have a valid OpenXR runtime installed."));
+		return false;
+	}
+
+	if (!EnumerateLayers())
+	{
+		UE_LOG(LogHMD, Log, TEXT("Failed to enumerate API layers. Please check that you have a valid OpenXR runtime installed."));
+		return false;
+	}
+
+	// Enable any required and optional extensions that are not plugin specific (usually platform support extensions)
+	{
+		TArray<const ANSICHAR*> RequiredExtensions, OptionalExtensions, Extensions;
+		// Query required extensions
+		RequiredExtensions.Empty();
+		if (!GetRequiredExtensions(RequiredExtensions))
+		{
+			UE_LOG(LogHMD, Error, TEXT("Could not get required OpenXR extensions."));
+			return false;
+		}
+
+		// Query optional extensions
+		OptionalExtensions.Empty();
+		if (!GetOptionalExtensions(OptionalExtensions))
+		{
+			UE_LOG(LogHMD, Error, TEXT("Could not get optional OpenXR extensions."));
+			return false;
+		}
+
+		if (!EnableExtensions(RequiredExtensions, OptionalExtensions, Extensions))
+		{
+			UE_LOG(LogHMD, Error, TEXT("Could not enable all required OpenXR extensions."));
+			return false;
+		}
+		ExtensionSet.Append(Extensions);
+	}
+
+	if (AvailableExtensions.Contains(XR_EPIC_VIEW_CONFIGURATION_FOV_EXTENSION_NAME))
+	{
+		ExtensionSet.Add(XR_EPIC_VIEW_CONFIGURATION_FOV_EXTENSION_NAME);
+	}
+
+	for (IOpenXRExtensionPlugin* Plugin : ExtModules)
+	{
+		TArray<const ANSICHAR*> RequiredExtensions, OptionalExtensions, Extensions;
+
+		// Query required extensions
+		RequiredExtensions.Empty();
+		if (!Plugin->GetRequiredExtensions(RequiredExtensions))
+		{
+			// Ignore the plugin if the query fails
+			continue;
+		}
+
+		// Query optional extensions
+		OptionalExtensions.Empty();
+		if (!Plugin->GetOptionalExtensions(OptionalExtensions))
+		{
+			// Ignore the plugin if the query fails
+			continue;
+		}
+
+		if (!EnableExtensions(RequiredExtensions, OptionalExtensions, Extensions))
+		{
+			// Ignore the plugin if the required extension could not be enabled
+			FString ModuleName = Plugin->GetDisplayName();
+			UE_LOG(LogHMD, Log, TEXT("Could not enable all required OpenXR extensions for %s on current system. This plugin will be loaded but ignored, but will be enabled on a target platform that supports the required extension."), *ModuleName);
+			continue;
+		}
+		ExtensionSet.Append(Extensions);
+		ExtensionPlugins.Add(Plugin);
+	}
+
+	if (auto ARModule = FModuleManager::LoadModulePtr<IOpenXRARModule>("OpenXRAR"))
+	{
+		TArray<const ANSICHAR*> ARExtensionSet;
+		ARModule->GetExtensions(ARExtensionSet);
+		ExtensionSet.Append(ARExtensionSet);
+	}
+
+	EnabledExtensions.Reset();
+	for (const ANSICHAR* Ext : ExtensionSet)
+	{
+		EnabledExtensions.Add(Ext);
+	}
+
+	// Enable layers, if specified by CVar.
+	// Note: For the validation layer to work on Windows (as of latest OpenXR runtime, August 2019), the following are required:
+	//   1. Download and build the OpenXR SDK from [REDACTED_PROFILE_URL]/OpenXR-SDK-Source (follow instructions at [REDACTED_PROFILE_URL]/OpenXR-SDK-Source/blob/master/BUILDING.md)
+	//	 2. Add a registry key under HKEY_LOCAL_MACHINE\SOFTWARE\Khronos\OpenXR\1\ApiLayers\Explicit, containing the path to the manifest file
+	//      (e.g. C:\OpenXR-SDK-Source-master\build\win64\src\api_layers\XrApiLayer_core_validation.json) <-- this file is downloaded as part of the SDK source, above
+	//   3. Copy the DLL from the build target at, for example, C:\OpenXR-SDK-Source-master\build\win64\src\api_layers\XrApiLayer_core_validation.dll to
+	//      somewhere in your system path (e.g. c:\windows\system32); the OpenXR loader currently doesn't use the path the json file is in (this is a bug)
+
+	const bool bEnableOpenXRValidationLayer = CVarEnableOpenXRValidationLayer.GetValueOnAnyThread() != 0;
+	TArray<const char*> Layers;
+	if (bEnableOpenXRValidationLayer && AvailableLayers.Contains("XR_APILAYER_LUNARG_core_validation"))
+	{
+		Layers.Add("XR_APILAYER_LUNARG_core_validation");
+	}
+
+	// Engine registration can be disabled via console var.
+	auto* CVarDisableEngineAndAppRegistration = IConsoleManager::Get().FindTConsoleVariableDataInt(TEXT("r.DisableEngineAndAppRegistration"));
+	bool bDisableEngineRegistration = (CVarDisableEngineAndAppRegistration && CVarDisableEngineAndAppRegistration->GetValueOnAnyThread() != 0);
+
+	FText ProjectName = FText();
+	GConfig->GetText(TEXT("/Script/EngineSettings.GeneralProjectSettings"), TEXT("ProjectName"), ProjectName, GGameIni);
+
+	FText ProjectVersion = FText();
+	GConfig->GetText(TEXT("/Script/EngineSettings.GeneralProjectSettings"), TEXT("ProjectVersion"), ProjectVersion, GGameIni);
+
+	// EngineName will be of the form "UnrealEngine4.21", with the minor version ("21" in this example)
+	// updated with every quarterly release
+	FString EngineName = bDisableEngineRegistration ? FString("") : FApp::GetEpicProductIdentifier() + FEngineVersion::Current().ToString(EVersionComponent::Minor);
+	FString AppName = bDisableEngineRegistration ? TEXT("") : ProjectName.ToString() + ProjectVersion.ToString();
+
+	XrInstanceCreateInfo Info;
+	Info.type = XR_TYPE_INSTANCE_CREATE_INFO;
+	Info.next = nullptr;
+	Info.createFlags = 0;
+	FTCHARToUTF8_Convert::Convert(Info.applicationInfo.applicationName, XR_MAX_APPLICATION_NAME_SIZE, *AppName, AppName.Len() + 1);
+	Info.applicationInfo.applicationVersion = static_cast<uint32>(BuildSettings::GetCurrentChangelist()) | (BuildSettings::IsLicenseeVersion() ? 0x80000000 : 0);
+	FTCHARToUTF8_Convert::Convert(Info.applicationInfo.engineName, XR_MAX_ENGINE_NAME_SIZE, *EngineName, EngineName.Len() + 1);
+	Info.applicationInfo.engineVersion = (uint32)(FEngineVersion::Current().GetMinor() << 16 | FEngineVersion::Current().GetPatch());
+	Info.applicationInfo.apiVersion = XR_CURRENT_API_VERSION;
+
+	Info.enabledApiLayerCount = Layers.Num();
+	Info.enabledApiLayerNames = Layers.GetData();
+
+	Info.enabledExtensionCount = EnabledExtensions.Num();
+	Info.enabledExtensionNames = EnabledExtensions.GetData();
+
+#if PLATFORM_ANDROID
+	XrInstanceCreateInfoAndroidKHR InstanceCreateInfoAndroid;
+	InstanceCreateInfoAndroid.type = XR_TYPE_INSTANCE_CREATE_INFO_ANDROID_KHR;
+	InstanceCreateInfoAndroid.next = nullptr;
+	InstanceCreateInfoAndroid.applicationVM = GNativeAndroidApp->activity->vm;
+	InstanceCreateInfoAndroid.applicationActivity = GNativeAndroidApp->activity->clazz;
+	Info.next = &InstanceCreateInfoAndroid;
+#endif // PLATFORM_ANDROID
+
+	for (IOpenXRExtensionPlugin* Module : ExtensionPlugins)
+	{
+		Info.next = Module->OnCreateInstance(this, Info.next);
+	}
+
+	XrResult Result = xrCreateInstance(&Info, &Instance);
+	if (XR_FAILED(Result))
+	{
+		UE_LOG(LogHMD, Log, TEXT("Failed to create an OpenXR instance, result is %s. Please check if you have an OpenXR runtime installed. The following extensions were enabled:"), OpenXRResultToString(Result));
+		for (const char* Extension : EnabledExtensions)
+		{
+			UE_LOG(LogHMD, Log, TEXT("- %S"), Extension);
+		}
+		return false;
+	}
+
+	if (!InitOpenXRCore(Instance))
+	{
+		UE_LOG(LogHMD, Log, TEXT("Failed to initialize core functions. Please check that you have a valid OpenXR runtime installed."));
+		return false;
+	}
+
+	for (IOpenXRExtensionPlugin* Module : ExtensionPlugins)
+	{
+		Module->PostCreateInstance(Instance);
+	}
+
+	return true;
+}
+
+bool FOpenXRHMDModule::InitSystem()
+{
+	XrSystemGetInfo SystemInfo;
+	SystemInfo.type = XR_TYPE_SYSTEM_GET_INFO;
+	SystemInfo.next = nullptr;
+	SystemInfo.formFactor = XR_FORM_FACTOR_HEAD_MOUNTED_DISPLAY;
+	for (IOpenXRExtensionPlugin* Module : ExtensionPlugins)
+	{
+		SystemInfo.next = Module->OnGetSystem(Instance, SystemInfo.next);
+	}
+
+	XrResult Result = xrGetSystem(Instance, &SystemInfo, &System);
+	if (XR_FAILED(Result))
+	{
+		UE_LOG(LogHMD, Log, TEXT("Failed to get an OpenXR system, result is %s. Please check that your runtime supports VR headsets."), OpenXRResultToString(Result));
+		return false;
+	}
+
+	for (IOpenXRExtensionPlugin* Module : ExtensionPlugins)
+	{
+		Module->PostGetSystem(Instance, System);
+	}
+
+	return true;
+}
diff --git a/Engine/Plugins/Runtime/OpenXR/Source/OpenXRHMD/Private/OpenXRHMDModule.h b/Engine/Plugins/Runtime/OpenXR/Source/OpenXRHMD/Private/OpenXRHMDModule.h
new file mode 100644
--- /dev/null
+++ b/Engine/Plugins/Runtime/OpenXR/Source/OpenXRHMD/Private/OpenXRHMDModule.h
@@ -0,0 +1,71 @@
+// Copyright Epic Games, Inc. All Rights Reserved.
+
+#pragma once
+
+#include "IOpenXRHMDModule.h"
+#include "IHeadMountedDisplayModule.h"
+#include "CoreMinimal.h"
+
+#include <openxr/openxr.h>
+
+class FOpenXRHMDModule : public IOpenXRHMDModule
+{
+public:
+	FOpenXRHMDModule()
+		: LoaderHandle(nullptr)
+		, Instance(XR_NULL_HANDLE)
+		, System(XR_NULL_SYSTEM_ID)
+		, RenderBridge(nullptr)
+	{ }
+
+	~FOpenXRHMDModule()
+	{
+	}
+
+	/** IHeadMountedDisplayModule implementation */
+	virtual TSharedPtr< class IXRTrackingSystem, ESPMode::ThreadSafe > CreateTrackingSystem() override;
+	virtual TSharedPtr< class IHeadMountedDisplayVulkanExtensions, ESPMode::ThreadSafe > GetVulkanExtensions() override;
+	virtual uint64 GetGraphicsAdapterLuid() override;
+
+	FString GetModuleKeyName() const override
+	{
+		return FString(TEXT("OpenXRHMD"));
+	}
+
+	void GetModuleAliases(TArray<FString>& AliasesOut) const override
+	{
+		AliasesOut.Add(TEXT("OpenXR"));
+	}
+
+	void ShutdownModule() override;
+
+	virtual bool IsHMDConnected() override { return true; }
+	virtual bool IsStandaloneStereoOnlyDevice() override;
+	virtual bool IsExtensionAvailable(const FString& Name) const override { return AvailableExtensions.Contains(Name); }
+	virtual bool IsExtensionEnabled(const FString& Name) const override { return EnabledExtensions.Contains(Name); }
+	virtual bool IsLayerAvailable(const FString& Name) const override { return EnabledLayers.Contains(Name); }
+	virtual bool IsLayerEnabled(const FString& Name) const override { return EnabledLayers.Contains(Name); }
+
+private:
+	void* LoaderHandle;
+	XrInstance Instance;
+	XrSystemId System;
+	TSet<FString> AvailableExtensions;
+	TSet<FString> AvailableLayers;
+	TArray<const char*> EnabledExtensions;
+	TArray<const char*> EnabledLayers;
+	TArray<class IOpenXRExtensionPlugin*> ExtensionPlugins;
+	TRefCountPtr<class FOpenXRRenderBridge> RenderBridge;
+	TSharedPtr< IHeadMountedDisplayVulkanExtensions, ESPMode::ThreadSafe > VulkanExtensions;
+
+	bool EnumerateExtensions();
+	bool EnumerateLayers();
+	bool InitRenderBridge();
+	bool InitInstanceAndSystem();
+	bool InitInstance();
+	bool InitSystem();
+	PFN_xrGetInstanceProcAddr GetDefaultLoader();
+	bool EnableExtensions(const TArray<const ANSICHAR*>& RequiredExtensions, const TArray<const ANSICHAR*>& OptionalExtensions, TArray<const ANSICHAR*>& OutExtensions);
+	bool GetRequiredExtensions(TArray<const ANSICHAR*>& OutExtensions);
+	bool GetOptionalExtensions(TArray<const ANSICHAR*>& OutExtensions);
+};
diff --git a/Engine/Plugins/Runtime/OpenXR/Source/OpenXRHMD/Public/IOpenXRExtensionPlugin.h b/Engine/Plugins/Runtime/OpenXR/Source/OpenXRHMD/Public/IOpenXRExtensionPlugin.h
--- a/Engine/Plugins/Runtime/OpenXR/Source/OpenXRHMD/Public/IOpenXRExtensionPlugin.h
+++ b/Engine/Plugins/Runtime/OpenXR/Source/OpenXRHMD/Public/IOpenXRExtensionPlugin.h
@@ -266,7 +266,7 @@ public:
 	* Remember to assign InNext to the next pointer of your struct or otherwise you may break the next chain.
 	*/
 
-	virtual const void* OnCreateInstance(class IOpenXRHMDPlugin* InPlugin, const void* InNext)
+	virtual const void* OnCreateInstance(class IOpenXRHMDModule* InModule, const void* InNext)
 	{
 		return InNext;
 	}
diff --git a/Engine/Plugins/Runtime/OpenXR/Source/OpenXRHMD/Public/IOpenXRHMDModule.h b/Engine/Plugins/Runtime/OpenXR/Source/OpenXRHMD/Public/IOpenXRHMDModule.h
new file mode 100644
--- /dev/null
+++ b/Engine/Plugins/Runtime/OpenXR/Source/OpenXRHMD/Public/IOpenXRHMDModule.h
@@ -0,0 +1,44 @@
+// Copyright Epic Games, Inc. All Rights Reserved.
+
+#pragma once
+
+#include "CoreMinimal.h"
+#include "Modules/ModuleManager.h"
+#include "IHeadMountedDisplayModule.h"
+
+/**
+ * The public interface to this module.  In most cases, this interface is only public to sibling modules 
+ * within this plugin.
+ */
+class OPENXRHMD_API IOpenXRHMDModule : public IHeadMountedDisplayModule
+{
+
+public:
+
+	/**
+	 * Singleton-like access to this module's interface.  This is just for convenience!
+	 * Beware of calling this during the shutdown phase, though.  Your module might have been unloaded already.
+	 *
+	 * @return Returns singleton instance, loading the module on demand if needed
+	 */
+	static inline IOpenXRHMDModule& Get()
+	{
+		return FModuleManager::LoadModuleChecked< IOpenXRHMDModule >( "OpenXRHMD" );
+	}
+
+	/**
+	 * Checks to see if this module is loaded and ready.  It is only valid to call Get() if IsAvailable() returns true.
+	 *
+	 * @return True if the module is loaded and ready to use
+	 */
+	static inline bool IsAvailable()
+	{
+		return FModuleManager::Get().IsModuleLoaded( "OpenXRHMD" );
+	}
+
+	virtual bool IsExtensionAvailable(const FString& Name) const = 0;
+	virtual bool IsExtensionEnabled(const FString& Name) const = 0;
+
+	virtual bool IsLayerAvailable(const FString& Name) const = 0;
+	virtual bool IsLayerEnabled(const FString& Name) const = 0;
+};
diff --git a/Engine/Plugins/Runtime/OpenXR/Source/OpenXRHMD/Public/IOpenXRHMDPlugin.h b/Engine/Plugins/Runtime/OpenXR/Source/OpenXRHMD/Public/IOpenXRHMDPlugin.h
deleted file mode 100644
diff --git a/Engine/Plugins/Runtime/OpenXRHandTracking/Source/OpenXRHandTracking/Private/OpenXRHandTracking.cpp b/Engine/Plugins/Runtime/OpenXRHandTracking/Source/OpenXRHandTracking/Private/OpenXRHandTracking.cpp
--- a/Engine/Plugins/Runtime/OpenXRHandTracking/Source/OpenXRHandTracking/Private/OpenXRHandTracking.cpp
+++ b/Engine/Plugins/Runtime/OpenXRHandTracking/Source/OpenXRHandTracking/Private/OpenXRHandTracking.cpp
@@ -17,7 +17,7 @@
 //#include "DrawDebugHelpers.h"
 #include "ILiveLinkClient.h"
 #include "HeadMountedDisplayFunctionLibrary.h"
-#include "IOpenXRHMDPlugin.h"
+#include "IOpenXRHMDModule.h"
 
 #define LOCTEXT_NAMESPACE "OpenXRHandTracking"
 
@@ -155,7 +155,7 @@ FOpenXRHandTracking::FOpenXRHandTracking(const TSharedRef<FGenericApplicationMes
 
 	// We're implicitly requiring that the OpenXRPlugin has been loaded and
 	// initialized at this point.
-	if (!IOpenXRHMDPlugin::Get().IsAvailable())
+	if (!IOpenXRHMDModule::Get().IsAvailable())
 	{
 		UE_LOG(LogOpenXRHandTracking, Error, TEXT("Error - OpenXRHMDPlugin isn't available"));
 	}
