Fix compile error in CL 17451221

diff --git a/Engine/Plugins/Runtime/OpenXR/Source/OpenXRHMD/Private/OpenXRHMDModule.cpp b/Engine/Plugins/Runtime/OpenXR/Source/OpenXRHMD/Private/OpenXRHMDModule.cpp
--- a/Engine/Plugins/Runtime/OpenXR/Source/OpenXRHMD/Private/OpenXRHMDModule.cpp
+++ b/Engine/Plugins/Runtime/OpenXR/Source/OpenXRHMD/Private/OpenXRHMDModule.cpp
@@ -26,6 +26,17 @@ static TAutoConsoleVariable<int32> CVarEnableOpenXRValidationLayer(
 
 IMPLEMENT_MODULE( FOpenXRHMDModule, OpenXRHMD )
 
+FOpenXRHMDModule::FOpenXRHMDModule()
+	: LoaderHandle(nullptr)
+	, Instance(XR_NULL_HANDLE)
+	, System(XR_NULL_SYSTEM_ID)
+	, RenderBridge(nullptr)
+{ }
+
+FOpenXRHMDModule::~FOpenXRHMDModule()
+{
+}
+
 TSharedPtr< class IXRTrackingSystem, ESPMode::ThreadSafe > FOpenXRHMDModule::CreateTrackingSystem()
 {
 	if (!RenderBridge)
diff --git a/Engine/Plugins/Runtime/OpenXR/Source/OpenXRHMD/Private/OpenXRHMDModule.h b/Engine/Plugins/Runtime/OpenXR/Source/OpenXRHMD/Private/OpenXRHMDModule.h
--- a/Engine/Plugins/Runtime/OpenXR/Source/OpenXRHMD/Private/OpenXRHMDModule.h
+++ b/Engine/Plugins/Runtime/OpenXR/Source/OpenXRHMD/Private/OpenXRHMDModule.h
@@ -11,16 +11,8 @@
 class FOpenXRHMDModule : public IOpenXRHMDModule
 {
 public:
-	FOpenXRHMDModule()
-		: LoaderHandle(nullptr)
-		, Instance(XR_NULL_HANDLE)
-		, System(XR_NULL_SYSTEM_ID)
-		, RenderBridge(nullptr)
-	{ }
-
-	~FOpenXRHMDModule()
-	{
-	}
+	FOpenXRHMDModule();
+	~FOpenXRHMDModule();
 
 	/** IHeadMountedDisplayModule implementation */
 	virtual TSharedPtr< class IXRTrackingSystem, ESPMode::ThreadSafe > CreateTrackingSystem() override;
