Fix difference between UE4/UE5 in OpenXR DSO loading file

diff --git a/Engine/Plugins/Runtime/OpenXR/Source/OpenXRHMD/Private/OpenXRHMDModule.cpp b/Engine/Plugins/Runtime/OpenXR/Source/OpenXRHMD/Private/OpenXRHMDModule.cpp
--- a/Engine/Plugins/Runtime/OpenXR/Source/OpenXRHMD/Private/OpenXRHMDModule.cpp
+++ b/Engine/Plugins/Runtime/OpenXR/Source/OpenXRHMD/Private/OpenXRHMDModule.cpp
@@ -290,15 +290,8 @@ PFN_xrGetInstanceProcAddr FOpenXRHMDModule::GetDefaultLoader()
 	LoaderHandle = FPlatformProcess::GetDllHandle(*LoaderName);
 	FPlatformProcess::PopDllDirectory(*BinariesPath);
 #elif PLATFORM_LINUX
-	LoaderHandle = FPlatformProcess::GetDllHandle(TEXT("/usr/lib/libopenxr_loader.so"));
-	if (!LoaderHandle)
-	{
-		LoaderHandle = FPlatformProcess::GetDllHandle(TEXT("/usr/lib/x86_64-linux-gnu/libopenxr_loader.so"));
-		if (!LoaderHandle)
-		{
-			LoaderHandle = FPlatformProcess::GetDllHandle(TEXT("/usr/local/lib/libopenxr_loader.so"));
-		}
-	}
+	FString BinariesPath = FPaths::EngineDir() / FString(TEXT("Binaries/ThirdParty/OpenXR/linux/x86_64-unknown-linux-gnu"));
+	LoaderHandle = FPlatformProcess::GetDllHandle(*(BinariesPath / TEXT("libopenxr_loader.so")));
 #elif PLATFORM_HOLOLENS
 #ifndef PLATFORM_64BITS
 #error HoloLens platform does not currently support 32-bit. 32-bit OpenXR loader binaries are needed.
