Fix D3D12 compile in UE5

diff --git a/Engine/Plugins/Runtime/OpenXR/Source/OpenXRHMD/Public/OpenXRPlatformRHI.h b/Engine/Plugins/Runtime/OpenXR/Source/OpenXRHMD/Public/OpenXRPlatformRHI.h
--- a/Engine/Plugins/Runtime/OpenXR/Source/OpenXRHMD/Public/OpenXRPlatformRHI.h
+++ b/Engine/Plugins/Runtime/OpenXR/Source/OpenXRHMD/Public/OpenXRPlatformRHI.h
@@ -35,22 +35,6 @@
 //-------------------------------------------------------------------------------------------------
 
 #ifdef XR_USE_GRAPHICS_API_D3D12
-#define GetD3D11CubeFace GetD3D12CubeFace
-#define VerifyD3D11Result VerifyD3D12Result
-#define GetD3D11TextureFromRHITexture GetD3D12TextureFromRHITexture
-#define FRingAllocation FRingAllocation_D3D12
-#define GetRenderTargetFormat GetRenderTargetFormat_D3D12
-#define ED3D11ShaderOffsetBuffer ED3D12ShaderOffsetBuffer
-#define FindShaderResourceDXGIFormat FindShaderResourceDXGIFormat_D3D12
-#define FindUnorderedAccessDXGIFormat FindUnorderedAccessDXGIFormat_D3D12
-#define FindDepthStencilDXGIFormat FindDepthStencilDXGIFormat_D3D12
-#define HasStencilBits HasStencilBits_D3D12
-#define FVector4VertexDeclaration FVector4VertexDeclaration_D3D12
-#define GLOBAL_CONSTANT_BUFFER_INDEX GLOBAL_CONSTANT_BUFFER_INDEX_D3D12
-#define MAX_CONSTANT_BUFFER_SLOTS MAX_CONSTANT_BUFFER_SLOTS_D3D12
-#define FD3DGPUProfiler FD3D12GPUProfiler
-#define FRangeAllocator FRangeAllocator_D3D12
-
 #include "D3D12RHIPrivate.h"
 #include "D3D12Util.h"
 
