Fix non-unity build

diff --git a/Engine/Plugins/Runtime/OpenXR/Source/OpenXRHMD/Private/OpenXRHMDModule.h b/Engine/Plugins/Runtime/OpenXR/Source/OpenXRHMD/Private/OpenXRHMDModule.h
--- a/Engine/Plugins/Runtime/OpenXR/Source/OpenXRHMD/Private/OpenXRHMDModule.h
+++ b/Engine/Plugins/Runtime/OpenXR/Source/OpenXRHMD/Private/OpenXRHMDModule.h
@@ -5,6 +5,7 @@
 #include "IOpenXRHMDModule.h"
 #include "IHeadMountedDisplayModule.h"
 #include "CoreMinimal.h"
+#include "Templates/RefCounting.h"
 
 #include <openxr/openxr.h>
 
