Fix swapchains not being released when the session stops running
Prevent swapchains from being updated while the session is not running
Always use BGRA8 for stereo layer swapchains

diff --git a/Engine/Plugins/Runtime/OpenXR/Source/OpenXRHMD/Private/OpenXRHMD.cpp b/Engine/Plugins/Runtime/OpenXR/Source/OpenXRHMD/Private/OpenXRHMD.cpp
--- a/Engine/Plugins/Runtime/OpenXR/Source/OpenXRHMD/Private/OpenXRHMD.cpp
+++ b/Engine/Plugins/Runtime/OpenXR/Source/OpenXRHMD/Private/OpenXRHMD.cpp
@@ -2223,7 +2223,7 @@ void FOpenXRHMD::OnBeginRendering_RenderThread(FRHICommandListImmediate& RHICmdL
 		// We need to copy each layer into an OpenXR swapchain so they can be displayed by the compositor
 		if (Layer.Swapchain.IsValid() && Layer.Desc.Texture.IsValid())
 		{
-			if (Layer.bUpdateTexture)
+			if (Layer.bUpdateTexture && bIsRunning)
 			{
 				FRHITexture2D* SrcTexture = Layer.Desc.Texture->GetTexture2D();
 				FIntRect DstRect(FIntPoint(0, 0), Layer.SwapchainSize.IntPoint());
@@ -2238,7 +2238,7 @@ void FOpenXRHMD::OnBeginRendering_RenderThread(FRHICommandListImmediate& RHICmdL
 
 		if (Layer.LeftSwapchain.IsValid() && Layer.Desc.LeftTexture.IsValid())
 		{
-			if (Layer.bUpdateTexture)
+			if (Layer.bUpdateTexture && bIsRunning)
 			{
 				FRHITexture2D* SrcTexture = Layer.Desc.LeftTexture->GetTexture2D();
 				FIntRect DstRect(FIntPoint(0, 0), Layer.SwapchainSize.IntPoint());
@@ -2252,12 +2252,6 @@ void FOpenXRHMD::OnBeginRendering_RenderThread(FRHICommandListImmediate& RHICmdL
 		}
 	}
 
-	// Reset the update flag on all layers
-	ForEachLayer([&](uint32 /* unused */, FOpenXRLayer& Layer)
-	{
-		Layer.bUpdateTexture = Layer.Desc.Flags & IStereoLayers::LAYER_FLAG_TEX_CONTINUOUS_UPDATE;
-	});
-
 #if !PLATFORM_HOLOLENS
 	if (bHiddenAreaMaskSupported && bNeedReBuildOcclusionMesh)
 	{
@@ -2269,6 +2263,12 @@ void FOpenXRHMD::OnBeginRendering_RenderThread(FRHICommandListImmediate& RHICmdL
 	{
 		SCOPED_NAMED_EVENT(EnqueueFrame, FColor::Red);
 
+		// Reset the update flag on all layers
+		ForEachLayer([&](uint32 /* unused */, FOpenXRLayer& Layer)
+		{
+			Layer.bUpdateTexture = Layer.Desc.Flags & IStereoLayers::LAYER_FLAG_TEX_CONTINUOUS_UPDATE;
+		});
+
 		if (bNeedsAcquireOnRHI)
 		{
 			// TODO: Current spec incorrectly requires this function to have access to the Vulkan queue.
@@ -2578,48 +2578,48 @@ void FOpenXRHMD::OnFinishRendering_RHIThread()
 
 	SCOPED_NAMED_EVENT(EndFrame, FColor::Red);
 
-	FReadScopeLock Lock(SessionHandleMutex);
-	if (!bIsRunning)
+	if (!bIsRendering)
 	{
 		return;
 	}
 
-	TArray<const XrCompositionLayerBaseHeader*> Headers;
-
-	XrCompositionLayerProjection Layer = {};
-	if (IsBackgroundLayerVisible())
+	// We need to ensure we release the swap chain images even if the session is not running.
+	if (PipelinedLayerStateRHI.ColorSwapchain)
 	{
-		Layer.type = XR_TYPE_COMPOSITION_LAYER_PROJECTION;
-		Layer.next = nullptr;
-		Layer.layerFlags = bProjectionLayerAlphaEnabled ? XR_COMPOSITION_LAYER_BLEND_TEXTURE_SOURCE_ALPHA_BIT : 0;
-		Layer.space = PipelinedFrameStateRHI.TrackingSpace;
-		Layer.viewCount = PipelinedLayerStateRHI.ProjectionLayers.Num();
-		Layer.views = PipelinedLayerStateRHI.ProjectionLayers.GetData();
-		Headers.Add(reinterpret_cast<const XrCompositionLayerBaseHeader*>(&Layer));
+		PipelinedLayerStateRHI.ColorSwapchain->ReleaseCurrentImage_RHIThread();
 
-		for (IOpenXRExtensionPlugin* Module : ExtensionPlugins)
+		if (bDepthExtensionSupported && PipelinedLayerStateRHI.DepthSwapchain)
 		{
-			Layer.next = Module->OnEndProjectionLayer(Session, 0, Layer.next, Layer.layerFlags);
+			PipelinedLayerStateRHI.DepthSwapchain->ReleaseCurrentImage_RHIThread();
 		}
 	}
 
-	for (const XrCompositionLayerQuad& Quad : PipelinedLayerStateRHI.QuadLayers)
-	{
-		Headers.Add(reinterpret_cast<const XrCompositionLayerBaseHeader*>(&Quad));
-	}
-
-	if (bIsRendering)
+	FReadScopeLock Lock(SessionHandleMutex);
+	if (bIsRunning)
 	{
-		if (PipelinedLayerStateRHI.ColorSwapchain)
+		TArray<const XrCompositionLayerBaseHeader*> Headers;
+		XrCompositionLayerProjection Layer = {};
+		if (IsBackgroundLayerVisible())
 		{
-			PipelinedLayerStateRHI.ColorSwapchain->ReleaseCurrentImage_RHIThread();
+			Layer.type = XR_TYPE_COMPOSITION_LAYER_PROJECTION;
+			Layer.next = nullptr;
+			Layer.layerFlags = bProjectionLayerAlphaEnabled ? XR_COMPOSITION_LAYER_BLEND_TEXTURE_SOURCE_ALPHA_BIT : 0;
+			Layer.space = PipelinedFrameStateRHI.TrackingSpace;
+			Layer.viewCount = PipelinedLayerStateRHI.ProjectionLayers.Num();
+			Layer.views = PipelinedLayerStateRHI.ProjectionLayers.GetData();
+			Headers.Add(reinterpret_cast<const XrCompositionLayerBaseHeader*>(&Layer));
 
-			if (bDepthExtensionSupported && PipelinedLayerStateRHI.DepthSwapchain)
+			for (IOpenXRExtensionPlugin* Module : ExtensionPlugins)
 			{
-				PipelinedLayerStateRHI.DepthSwapchain->ReleaseCurrentImage_RHIThread();
+				Layer.next = Module->OnEndProjectionLayer(Session, 0, Layer.next, Layer.layerFlags);
 			}
 		}
 
+		for (const XrCompositionLayerQuad& Quad : PipelinedLayerStateRHI.QuadLayers)
+		{
+			Headers.Add(reinterpret_cast<const XrCompositionLayerBaseHeader*>(&Quad));
+		}
+
 		XrFrameEndInfo EndInfo;
 		EndInfo.type = XR_TYPE_FRAME_END_INFO;
 		EndInfo.next = nullptr;
@@ -2647,9 +2647,9 @@ void FOpenXRHMD::OnFinishRendering_RHIThread()
 			EndInfo.next = Module->OnEndFrame(Session, EndInfo.displayTime, ColorImages, DepthImages, EndInfo.next);
 		}
 		XR_ENSURE(xrEndFrame(Session, &EndInfo));
-
-		bIsRendering = false;
 	}
+
+	bIsRendering = false;
 }
 
 FXRRenderBridge* FOpenXRHMD::GetActiveRenderBridge_GameThread(bool /* bUseSeparateRenderTarget */)
@@ -2865,13 +2865,27 @@ void FOpenXRHMD::UpdateLayer(FOpenXRLayer& Layer, uint32 LayerId, bool bIsValid)
 
 	if (bIsValid)
 	{
-		ETextureCreateFlags Flags = Layer.Desc.Flags & IStereoLayers::LAYER_FLAG_TEX_CONTINUOUS_UPDATE ?
+		const ETextureCreateFlags Flags = Layer.Desc.Flags & IStereoLayers::LAYER_FLAG_TEX_CONTINUOUS_UPDATE ?
 			TexCreate_Dynamic : TexCreate_None;
 
+		auto CreateSwapchain = [this](FRHITexture2D* Texture, ETextureCreateFlags Flags)
+		{
+			return RenderBridge->CreateSwapchain(Session,
+				PF_B8G8R8A8,
+				Texture->GetSizeX(),
+				Texture->GetSizeY(),
+				1,
+				Texture->GetNumMips(),
+				Texture->GetNumSamples(),
+				Texture->GetFlags() | Flags,
+				TexCreate_RenderTargetable,
+				Texture->GetClearBinding());
+		};
+
 		if (Layer.NeedReAllocateTexture())
 		{
 			FRHITexture2D* Texture = Layer.Desc.Texture->GetTexture2D();
-			Layer.Swapchain = RenderBridge->CreateSwapchain(Session, Texture, Flags, TexCreate_RenderTargetable);
+			Layer.Swapchain = CreateSwapchain(Texture, Flags);
 			Layer.SwapchainSize = Texture->GetSizeXY();
 			Layer.bUpdateTexture = true;
 		}
@@ -2879,7 +2893,7 @@ void FOpenXRHMD::UpdateLayer(FOpenXRLayer& Layer, uint32 LayerId, bool bIsValid)
 		if (Layer.NeedReAllocateLeftTexture())
 		{
 			FRHITexture2D* Texture = Layer.Desc.LeftTexture->GetTexture2D();
-			Layer.LeftSwapchain = RenderBridge->CreateSwapchain(Session, Texture, Flags, TexCreate_RenderTargetable);
+			Layer.LeftSwapchain = CreateSwapchain(Texture, Flags);
 			Layer.bUpdateTexture = true;
 		}
 	}
diff --git a/Engine/Plugins/Runtime/OpenXR/Source/OpenXRHMD/Private/OpenXRHMD_RenderBridge.h b/Engine/Plugins/Runtime/OpenXR/Source/OpenXRHMD/Private/OpenXRHMD_RenderBridge.h
--- a/Engine/Plugins/Runtime/OpenXR/Source/OpenXRHMD/Private/OpenXRHMD_RenderBridge.h
+++ b/Engine/Plugins/Runtime/OpenXR/Source/OpenXRHMD/Private/OpenXRHMD_RenderBridge.h
@@ -26,25 +26,6 @@ public:
 	virtual void* GetGraphicsBinding() = 0;
 	virtual FXRSwapChainPtr CreateSwapchain(XrSession InSession, uint8 Format, uint32 SizeX, uint32 SizeY, uint32 ArraySize, uint32 NumMips, uint32 NumSamples, ETextureCreateFlags Flags, ETextureCreateFlags TargetableTextureFlags, const FClearValueBinding& ClearValueBinding) = 0;
 
-	FXRSwapChainPtr CreateSwapchain(XrSession InSession, FRHITexture2D* Template, ETextureCreateFlags Flags, ETextureCreateFlags TargetableTextureFlags)
-	{
-		if (!Template)
-		{
-			return nullptr;
-		}
-
-		return CreateSwapchain(InSession,
-			Template->GetFormat(),
-			Template->GetSizeX(),
-			Template->GetSizeY(),
-			1,
-			Template->GetNumMips(),
-			Template->GetNumSamples(),
-			Template->GetFlags() | Flags,
-			TargetableTextureFlags,
-			Template->GetClearBinding());
-	}
-
 	/** FRHICustomPresent */
 	virtual bool Present(int32& InOutSyncInterval) override;
 
