Allow OnControllerButtonReleased to be triggered when isActive becomes false

diff --git a/Engine/Plugins/Runtime/OpenXR/Source/OpenXRInput/Private/OpenXRInput.cpp b/Engine/Plugins/Runtime/OpenXR/Source/OpenXRInput/Private/OpenXRInput.cpp
--- a/Engine/Plugins/Runtime/OpenXR/Source/OpenXRInput/Private/OpenXRInput.cpp
+++ b/Engine/Plugins/Runtime/OpenXR/Source/OpenXRInput/Private/OpenXRInput.cpp
@@ -594,9 +594,9 @@ void FOpenXRInputPlugin::FOpenXRInput::SendControllerEvents()
 				State.type = XR_TYPE_ACTION_STATE_BOOLEAN;
 				State.next = nullptr;
 				XrResult Result = xrGetActionStateBoolean(Session, &GetInfo, &State);
-				if (Result >= XR_SUCCESS && State.isActive && State.changedSinceLastSync)
+				if (Result >= XR_SUCCESS && State.changedSinceLastSync)
 				{
-					if (State.currentState)
+					if (State.currentState && State.isActive)
 					{
 						MessageHandler->OnControllerButtonPressed(*ActionKey, 0, /*IsRepeat =*/false);
 					}
