Re-enable the OpenXR depth extension on WMR with D3D12
Issue was caused by RenderDoc no artifacts or errors with plugin disabled

diff --git a/Engine/Plugins/Runtime/OpenXR/Source/OpenXRHMD/Private/OpenXRHMD.cpp b/Engine/Plugins/Runtime/OpenXR/Source/OpenXRHMD/Private/OpenXRHMD.cpp
--- a/Engine/Plugins/Runtime/OpenXR/Source/OpenXRHMD/Private/OpenXRHMD.cpp
+++ b/Engine/Plugins/Runtime/OpenXR/Source/OpenXRHMD/Private/OpenXRHMD.cpp
@@ -1398,11 +1398,6 @@ bool CheckPlatformDepthExtensionSupport(const XrInstanceProperties& InstanceProp
 		// No PF_DepthStencil compatible formats offered yet
 		return false;
 	}
-	else if (FCStringAnsi::Strstr(InstanceProps.runtimeName, "Windows Mixed Reality Runtime") && (FApp::GetGraphicsRHI() == TEXT("DirectX 12")))
-	{
-		// Temp: WMR missing depth layout transitions
-		return false;
-	}
 	return true;
 }
 
