Fix StereoLayer swapchains not being cleared before updating

diff --git a/Engine/Plugins/Runtime/OpenXR/Source/OpenXRHMD/Private/OpenXRHMD.cpp b/Engine/Plugins/Runtime/OpenXR/Source/OpenXRHMD/Private/OpenXRHMD.cpp
--- a/Engine/Plugins/Runtime/OpenXR/Source/OpenXRHMD/Private/OpenXRHMD.cpp
+++ b/Engine/Plugins/Runtime/OpenXR/Source/OpenXRHMD/Private/OpenXRHMD.cpp
@@ -2687,7 +2687,7 @@ FIntRect FOpenXRHMD::GetFullFlatEyeRect_RenderThread(FTexture2DRHIRef EyeTexture
 	return FIntRect(EyeTexture->GetSizeX() * SrcNormRectMin.X, EyeTexture->GetSizeY() * SrcNormRectMin.Y, EyeTexture->GetSizeX() * SrcNormRectMax.X, EyeTexture->GetSizeY() * SrcNormRectMax.Y);
 }
 
-void FOpenXRHMD::CopyTexture_RenderThread(FRHICommandListImmediate& RHICmdList, FRHITexture2D* SrcTexture, FIntRect SrcRect, FRHITexture2D* DstTexture, FIntRect DstRect, bool bClearBlack, bool bNoAlpha) const
+void FOpenXRHMD::CopyTexture_RenderThread(FRHICommandListImmediate& RHICmdList, FRHITexture2D* SrcTexture, FIntRect SrcRect, FRHITexture2D* DstTexture, FIntRect DstRect, bool bClearBlack, bool bNoAlpha, ERenderTargetActions RTAction) const
 {
 	check(IsInRenderingThread());
 
@@ -2714,7 +2714,7 @@ void FOpenXRHMD::CopyTexture_RenderThread(FRHICommandListImmediate& RHICmdList,
 	}
 
 	FRHITexture * ColorRT = DstTexture->GetTexture2DArray() ? DstTexture->GetTexture2DArray() : DstTexture->GetTexture2D();
-	FRHIRenderPassInfo RenderPassInfo(ColorRT, ERenderTargetActions::DontLoad_Store);
+	FRHIRenderPassInfo RenderPassInfo(ColorRT, RTAction);
 	RHICmdList.BeginRenderPass(RenderPassInfo, TEXT("OpenXRHMD_CopyTexture"));
 	{
 		if (bClearBlack)
@@ -2794,7 +2794,7 @@ void FOpenXRHMD::CopyTexture_RenderThread(FRHICommandListImmediate& RHICmdList,
 
 	// Now that we've enqueued the swapchain wait we can add the commands to do the actual texture copy
 	FRHITexture2DArray* DstTexture = DstSwapChain->GetTexture2DArray();
-	CopyTexture_RenderThread(RHICmdList, SrcTexture, SrcRect, DstTexture, DstRect, bClearBlack, bNoAlpha);
+	CopyTexture_RenderThread(RHICmdList, SrcTexture, SrcRect, DstTexture, DstRect, bClearBlack, bNoAlpha, ERenderTargetActions::Clear_Store);
 
 	// Enqueue a command to release the image after the copy is done
 	RHICmdList.EnqueueLambda([DstSwapChain](FRHICommandListImmediate& InRHICmdList)
@@ -2803,6 +2803,12 @@ void FOpenXRHMD::CopyTexture_RenderThread(FRHICommandListImmediate& RHICmdList,
 	});
 }
 
+void FOpenXRHMD::CopyTexture_RenderThread(FRHICommandListImmediate& RHICmdList, FRHITexture2D* SrcTexture, FIntRect SrcRect, FRHITexture2D* DstTexture, FIntRect DstRect, bool bClearBlack, bool bNoAlpha) const
+{
+	// FIXME: This should probably use the Load_Store action since the spectator controller does multiple overlaying copies.
+	CopyTexture_RenderThread(RHICmdList, SrcTexture, SrcRect, DstTexture, DstRect, bClearBlack, bNoAlpha, ERenderTargetActions::DontLoad_Store);
+}
+
 void FOpenXRHMD::RenderTexture_RenderThread(class FRHICommandListImmediate& RHICmdList, class FRHITexture2D* BackBuffer, class FRHITexture2D* SrcTexture, FVector2D WindowSize) const
 {
 	if (SpectatorScreenController)
diff --git a/Engine/Plugins/Runtime/OpenXR/Source/OpenXRHMD/Private/OpenXRHMD.h b/Engine/Plugins/Runtime/OpenXR/Source/OpenXRHMD/Private/OpenXRHMD.h
--- a/Engine/Plugins/Runtime/OpenXR/Source/OpenXRHMD/Private/OpenXRHMD.h
+++ b/Engine/Plugins/Runtime/OpenXR/Source/OpenXRHMD/Private/OpenXRHMD.h
@@ -182,6 +182,7 @@ protected:
 	void UpdateDeviceLocations(bool bUpdateOpenXRExtensionPlugins);
 	void EnumerateViews(FPipelinedFrameState& PipelineState);
 
+	void CopyTexture_RenderThread(FRHICommandListImmediate& RHICmdList, FRHITexture2D* SrcTexture, FIntRect SrcRect, FRHITexture2D* DstTexture, FIntRect DstRect, bool bClearBlack, bool bNoAlpha, ERenderTargetActions RTAction) const;
 	void CopyTexture_RenderThread(FRHICommandListImmediate& RHICmdList, FRHITexture2D* SrcTexture, FIntRect SrcRect, const FXRSwapChainPtr& DstSwapChain, FIntRect DstRect, bool bClearBlack, bool bNoAlpha) const;
 
 public:
