PR #7888: [OpenXR] Plugin interface changes to allow simple foveated rendering (Contributed by [REDACTED_PERSON])

diff --git a/Engine/Plugins/Runtime/OpenXR/Source/OpenXRHMD/Private/OpenXRHMD.cpp b/Engine/Plugins/Runtime/OpenXR/Source/OpenXRHMD/Private/OpenXRHMD.cpp
--- a/Engine/Plugins/Runtime/OpenXR/Source/OpenXRHMD/Private/OpenXRHMD.cpp
+++ b/Engine/Plugins/Runtime/OpenXR/Source/OpenXRHMD/Private/OpenXRHMD.cpp
@@ -690,7 +690,7 @@ bool FOpenXRHMDPlugin::InitInstance()
 	SystemInfo.formFactor = XR_FORM_FACTOR_HEAD_MOUNTED_DISPLAY;
 	for (IOpenXRExtensionPlugin* Module : ExtensionPlugins)
 	{
-		Info.next = Module->OnGetSystem(Instance, Info.next);
+		SystemInfo.next = Module->OnGetSystem(Instance, SystemInfo.next);
 	}
 	rs = xrGetSystem(Instance, &SystemInfo, &System);
 	if (XR_FAILED(rs))
@@ -1629,6 +1629,11 @@ void FOpenXRHMD::EnumerateViews(FPipelinedFrameState& PipelineState)
 		ViewFov[ViewIndex].next = nullptr;
 		View.next = bViewConfigurationFovSupported ? &ViewFov[ViewIndex] : nullptr;
 
+		for (IOpenXRExtensionPlugin* Module : ExtensionPlugins)
+		{
+			View.next = Module->OnEnumerateViewConfigurationViews(Instance, System, SelectedViewConfigurationType, ViewIndex, View.next);
+		}
+
 		// These are core views that don't have an associated plugin
 		PipelineState.PluginViews.Add(nullptr);
 		PipelineState.ViewConfigs.Add(View);
@@ -1655,6 +1660,10 @@ void FOpenXRHMD::EnumerateViews(FPipelinedFrameState& PipelineState)
 		ViewInfo.viewConfigurationType = SelectedViewConfigurationType;
 		ViewInfo.space = DeviceSpaces[HMDDeviceId].Space;
 		ViewInfo.displayTime = PipelineState.FrameState.predictedDisplayTime;
+		for (IOpenXRExtensionPlugin* Module : ExtensionPlugins)
+		{
+			ViewInfo.next = Module->OnLocateViews(Session, ViewInfo.displayTime, ViewInfo.next);
+		}
 		XR_ENSURE(xrLocateViews(Session, &ViewInfo, &PipelineState.ViewState, 0, &ViewCount, nullptr));
 		PipelineState.Views.SetNum(ViewCount, false);
 		XR_ENSURE(xrLocateViews(Session, &ViewInfo, &PipelineState.ViewState, PipelineState.Views.Num(), &ViewCount, PipelineState.Views.GetData()));
@@ -1828,6 +1837,11 @@ bool FOpenXRHMD::OnStereoStartup()
 	}
 	XR_ENSURE(xrCreateSession(Instance, &SessionInfo, &Session));
 
+	for (IOpenXRExtensionPlugin* Module : ExtensionPlugins)
+	{
+		Module->PostCreateSession(Session);
+	}
+
 	uint32_t ReferenceSpacesCount;
 	XR_ENSURE(xrEnumerateReferenceSpaces(Session, 0, &ReferenceSpacesCount, nullptr));
 
diff --git a/Engine/Plugins/Runtime/OpenXR/Source/OpenXRHMD/Public/IOpenXRExtensionPlugin.h b/Engine/Plugins/Runtime/OpenXR/Source/OpenXRHMD/Public/IOpenXRExtensionPlugin.h
--- a/Engine/Plugins/Runtime/OpenXR/Source/OpenXRHMD/Public/IOpenXRExtensionPlugin.h
+++ b/Engine/Plugins/Runtime/OpenXR/Source/OpenXRHMD/Public/IOpenXRExtensionPlugin.h
@@ -116,11 +116,16 @@ public:
 	* It is common to do this in StartupModule of your IModuleInterface class (which may also be the class that implements this interface).
 	* The module's LoadingPhase must be PostInitConfig or earlier because OpenXRHMD will look for these after it is loaded in that phase.
 	*/
-	virtual void RegisterOpenXRExtensionModularFeature()
+	void RegisterOpenXRExtensionModularFeature()
 	{
 		IModularFeatures::Get().RegisterModularFeature(GetModularFeatureName(), this);
 	}
 
+	void UnregisterOpenXRExtensionModularFeature()
+	{
+		IModularFeatures::Get().UnregisterModularFeature(GetModularFeatureName(), this);
+	}
+
 	/**
 	* Optionally provide a custom loader for the OpenXR plugin.
 	*/
@@ -194,6 +199,17 @@ public:
 
 	/** Get custom capture interface if provided by this extension. */
 	virtual IOpenXRCustomCaptureSupport* GetCustomCaptureSupport(const EARCaptureType CaptureType) { return nullptr; }
+
+	virtual void* OnEnumerateViewConfigurationViews(XrInstance InInstance, XrSystemId InSystem, XrViewConfigurationType InViewConfigurationType, uint32_t InViewIndex, void* InNext)
+	{
+		return InNext;
+	}
+
+	virtual const void* OnLocateViews(XrSession InSession, XrTime InDisplayTime, const void* InNext)
+	{
+		return InNext;
+	}
+
 	/**
 	* Callback to provide extra view configurations that should be rendered in the main render pass
 	*/
@@ -232,6 +248,10 @@ public:
 		return InNext;
 	}
 
+	virtual void PostCreateSession(XrSession InSession)
+	{
+	}
+
 	virtual const void* OnBeginSession(XrSession InSession, const void* InNext)
 	{
 		return InNext;
