Ensure StartSession/StopSession use exclusive lock so frame calls can't overlap session calls.

diff --git a/Engine/Plugins/Runtime/OpenXR/Source/OpenXRHMD/Private/OpenXRHMD.cpp b/Engine/Plugins/Runtime/OpenXR/Source/OpenXRHMD/Private/OpenXRHMD.cpp
--- a/Engine/Plugins/Runtime/OpenXR/Source/OpenXRHMD/Private/OpenXRHMD.cpp
+++ b/Engine/Plugins/Runtime/OpenXR/Source/OpenXRHMD/Private/OpenXRHMD.cpp
@@ -2055,7 +2055,7 @@ bool FOpenXRHMD::IsFocused() const
 bool FOpenXRHMD::StartSession()
 {
 	// If the session is not yet ready, we'll call into this function again when it is
-	FReadScopeLock Lock(SessionHandleMutex);
+	FWriteScopeLock Lock(SessionHandleMutex);
 	if (!bIsReady || bIsRunning)
 	{
 		return false;
@@ -2072,7 +2072,7 @@ bool FOpenXRHMD::StartSession()
 
 bool FOpenXRHMD::StopSession()
 {
-	FReadScopeLock Lock(SessionHandleMutex);
+	FWriteScopeLock Lock(SessionHandleMutex);
 	if (!bIsRunning)
 	{
 		return false;
