Editor crashes when opening a project with OpenXR enabled when an HMD is not attached.
FOpenXRHMDPlugin::InitInstance() is called twice (once for GetGraphicsAdapterLuid, then again for CreateTrackingSystem). On the second time, FOpenXRHMDPlugin::InitInstance will succeed because a valid Instance was created. Later, many OpenXR calls will fail because a valid xrSystem does not exist since InitInstance was responsible for creating it.
Refactor how the initialization works to ensure we avoid the crash, and also clean up the codepath slightly.

diff --git a/Engine/Plugins/Runtime/OpenXR/Source/OpenXRHMD/Private/OpenXRHMD.cpp b/Engine/Plugins/Runtime/OpenXR/Source/OpenXRHMD/Private/OpenXRHMD.cpp
--- a/Engine/Plugins/Runtime/OpenXR/Source/OpenXRHMD/Private/OpenXRHMD.cpp
+++ b/Engine/Plugins/Runtime/OpenXR/Source/OpenXRHMD/Private/OpenXRHMD.cpp
@@ -156,7 +156,9 @@ private:
 	bool EnumerateExtensions();
 	bool EnumerateLayers();
 	bool InitRenderBridge();
+	bool InitInstanceAndSystem();
 	bool InitInstance();
+	bool InitSystem();
 	PFN_xrGetInstanceProcAddr GetDefaultLoader();
 	bool EnableExtensions(const TArray<const ANSICHAR*>& RequiredExtensions, const TArray<const ANSICHAR*>& OptionalExtensions, TArray<const ANSICHAR*>& OutExtensions);
 	bool GetRequiredExtensions(TArray<const ANSICHAR*>& OutExtensions);
@@ -203,7 +205,7 @@ uint64 FOpenXRHMDPlugin::GetGraphicsAdapterLuid()
 TSharedPtr< IHeadMountedDisplayVulkanExtensions, ESPMode::ThreadSafe > FOpenXRHMDPlugin::GetVulkanExtensions()
 {
 #ifdef XR_USE_GRAPHICS_API_VULKAN
-	if (InitInstance() && IsExtensionEnabled(XR_KHR_VULKAN_ENABLE_EXTENSION_NAME))
+	if (InitInstanceAndSystem() && IsExtensionEnabled(XR_KHR_VULKAN_ENABLE_EXTENSION_NAME))
 	{
 		if (!VulkanExtensions.IsValid())
 		{
@@ -217,7 +219,7 @@ TSharedPtr< IHeadMountedDisplayVulkanExtensions, ESPMode::ThreadSafe > FOpenXRHM
 
 bool FOpenXRHMDPlugin::IsStandaloneStereoOnlyDevice()
 {
-	if (InitInstance())
+	if (InitInstanceAndSystem())
 	{
 		for (IOpenXRExtensionPlugin* Module : ExtensionPlugins)
 		{
@@ -299,7 +301,7 @@ bool FOpenXRHMDPlugin::InitRenderBridge()
 		return false;
 	}
 
-	if (!InitInstance())
+	if (!InitInstanceAndSystem())
 	{
 		return false;
 	}
@@ -489,13 +491,26 @@ bool FOpenXRHMDPlugin::GetOptionalExtensions(TArray<const ANSICHAR*>& OutExtensi
 	return true;
 }
 
-bool FOpenXRHMDPlugin::InitInstance()
+bool FOpenXRHMDPlugin::InitInstanceAndSystem()
 {
-	if (Instance)
+	if (!Instance && !InitInstance())
 	{
-		return true;
+		return false;
+	}
+
+	if (!System && !InitSystem())
+	{
+		return false;
 	}
 
+	return true;
+}
+
+bool FOpenXRHMDPlugin::InitInstance()
+{
+	// This should only ever be called if we don't already have an instance.
+	check(!Instance);
+
 	// Get all extension plugins
 	TSet<const ANSICHAR*, AnsiKeyFunc> ExtensionSet;
 	TArray<IOpenXRExtensionPlugin*> ExtModules = IModularFeatures::Get().GetModularFeatureImplementations<IOpenXRExtensionPlugin>(IOpenXRExtensionPlugin::GetModularFeatureName());
@@ -671,10 +686,11 @@ bool FOpenXRHMDPlugin::InitInstance()
 	{
 		Info.next = Module->OnCreateInstance(this, Info.next);
 	}
-	XrResult rs = xrCreateInstance(&Info, &Instance);
-	if (XR_FAILED(rs))
+
+	XrResult Result = xrCreateInstance(&Info, &Instance);
+	if (XR_FAILED(Result))
 	{
-		UE_LOG(LogHMD, Log, TEXT("Failed to create an OpenXR instance, result is %s. Please check if you have an OpenXR runtime installed."), OpenXRResultToString(rs));
+		UE_LOG(LogHMD, Log, TEXT("Failed to create an OpenXR instance, result is %s. Please check if you have an OpenXR runtime installed."), OpenXRResultToString(Result));
 		return false;
 	}
 
@@ -684,6 +700,11 @@ bool FOpenXRHMDPlugin::InitInstance()
 		return false;
 	}
 
+	return true;
+}
+
+bool FOpenXRHMDPlugin::InitSystem()
+{
 	XrSystemGetInfo SystemInfo;
 	SystemInfo.type = XR_TYPE_SYSTEM_GET_INFO;
 	SystemInfo.next = nullptr;
@@ -692,12 +713,14 @@ bool FOpenXRHMDPlugin::InitInstance()
 	{
 		SystemInfo.next = Module->OnGetSystem(Instance, SystemInfo.next);
 	}
-	rs = xrGetSystem(Instance, &SystemInfo, &System);
-	if (XR_FAILED(rs))
+
+	XrResult Result = xrGetSystem(Instance, &SystemInfo, &System);
+	if (XR_FAILED(Result))
 	{
-		UE_LOG(LogHMD, Log, TEXT("Failed to get an OpenXR system, result is %s. Please check that your runtime supports VR headsets."), OpenXRResultToString(rs));
+		UE_LOG(LogHMD, Log, TEXT("Failed to get an OpenXR system, result is %s. Please check that your runtime supports VR headsets."), OpenXRResultToString(Result));
 		return false;
 	}
+
 	for (IOpenXRExtensionPlugin* Module : ExtensionPlugins)
 	{
 		Module->PostGetSystem(Instance, System);
