UEVR-2213 OpenXR extensibility refactors

diff --git a/Engine/Plugins/Runtime/Oculus/OculusOpenXR/Source/OculusOpenXRHMD/Public/IOculusOpenXRHMDPlugin.h b/Engine/Plugins/Runtime/Oculus/OculusOpenXR/Source/OculusOpenXRHMD/Public/IOculusOpenXRHMDPlugin.h
--- a/Engine/Plugins/Runtime/Oculus/OculusOpenXR/Source/OculusOpenXRHMD/Public/IOculusOpenXRHMDPlugin.h
+++ b/Engine/Plugins/Runtime/Oculus/OculusOpenXR/Source/OculusOpenXRHMD/Public/IOculusOpenXRHMDPlugin.h
@@ -2,6 +2,7 @@
 
 #pragma once
 #include "IOpenXRExtensionPlugin.h"
+#include "Modules/ModuleInterface.h"
 
 class IOculusOpenXRHMDPlugin : public IOpenXRExtensionPlugin, public IModuleInterface
 {
diff --git a/Engine/Plugins/Runtime/OpenXR/Source/OpenXRHMD/OpenXRHMD.Build.cs b/Engine/Plugins/Runtime/OpenXR/Source/OpenXRHMD/OpenXRHMD.Build.cs
--- a/Engine/Plugins/Runtime/OpenXR/Source/OpenXRHMD/OpenXRHMD.Build.cs
+++ b/Engine/Plugins/Runtime/OpenXR/Source/OpenXRHMD/OpenXRHMD.Build.cs
@@ -25,12 +25,17 @@ namespace UnrealBuildTool.Rules
             {
                 PrivateIncludePaths.Add(EngineDir + "/Source/Runtime/VulkanRHI/Private/Windows");
             }
-            else if (Target.Platform != UnrealTargetPlatform.HoloLens)
+			else if (Target.Platform == UnrealTargetPlatform.Android  || Target.Platform == UnrealTargetPlatform.Linux)
             {
                 PrivateIncludePaths.Add(EngineDir + "/Source/Runtime/VulkanRHI/Private/" + Target.Platform);
             }
 
-            PrivateDependencyModuleNames.AddRange(
+			PublicIncludePathModuleNames.Add("OpenXR");
+			
+
+            PublicDependencyModuleNames.Add("HeadMountedDisplay");
+
+			PrivateDependencyModuleNames.AddRange(
 				new string[]
 				{
 					"Core",
@@ -43,7 +48,6 @@ namespace UnrealBuildTool.Rules
 					"RenderCore",
 					"Renderer",
 					"RenderCore",
-                    "HeadMountedDisplay",
                     "Slate",
                     "SlateCore",
 					"AugmentedReality",
@@ -94,7 +98,7 @@ namespace UnrealBuildTool.Rules
                 AddEngineThirdPartyPrivateStaticDependencies(Target, "OpenGL");
 			}
 
-            if (Target.Platform != UnrealTargetPlatform.HoloLens)
+			if (Target.Platform == UnrealTargetPlatform.Win32 || Target.Platform == UnrealTargetPlatform.Win64 || Target.Platform == UnrealTargetPlatform.Android)
             {
                 PrivateDependencyModuleNames.AddRange(new string[] {
                     "VulkanRHI"
@@ -102,9 +106,6 @@ namespace UnrealBuildTool.Rules
 
                 AddEngineThirdPartyPrivateStaticDependencies(Target, "Vulkan");
             }
-
-            PublicDependencyModuleNames.Add("HeadMountedDisplay");
-			PublicIncludePathModuleNames.Add("OpenXR");
 		}
 	}
 }
diff --git a/Engine/Plugins/Runtime/OpenXR/Source/OpenXRHMD/Private/OpenXRHMD.cpp b/Engine/Plugins/Runtime/OpenXR/Source/OpenXRHMD/Private/OpenXRHMD.cpp
--- a/Engine/Plugins/Runtime/OpenXR/Source/OpenXRHMD/Private/OpenXRHMD.cpp
+++ b/Engine/Plugins/Runtime/OpenXR/Source/OpenXRHMD/Private/OpenXRHMD.cpp
@@ -293,8 +293,55 @@ bool FOpenXRHMDPlugin::EnumerateLayers()
 	return false;
 }
 
+struct AnsiKeyFunc : BaseKeyFuncs<const ANSICHAR*, const ANSICHAR*, false>
+{
+	typedef typename TTypeTraits<const ANSICHAR*>::ConstPointerType KeyInitType;
+	typedef typename TCallTraits<const ANSICHAR*>::ParamType ElementInitType;
+
+	/**
+	 * @return The key used to index the given element.
+	 */
+	static FORCEINLINE KeyInitType GetSetKey(ElementInitType Element)
+	{
+		return Element;
+	}
+
+	/**
+	 * @return True if the keys match.
+	 */
+	static FORCEINLINE bool Matches(KeyInitType A, KeyInitType B)
+	{
+		return FCStringAnsi::Strcmp(A, B) == 0;
+	}
+
+	/** Calculates a hash index for a key. */
+	static FORCEINLINE uint32 GetKeyHash(KeyInitType Key)
+	{
+		return GetTypeHash(Key);
+	}
+};
+
 bool FOpenXRHMDPlugin::InitRenderBridge()
 {
+	// Get all extension plugins
+	TSet<const ANSICHAR*, AnsiKeyFunc> ExtensionSet;
+	TArray<IOpenXRExtensionPlugin*> ExtModules = IModularFeatures::Get().GetModularFeatureImplementations<IOpenXRExtensionPlugin>(IOpenXRExtensionPlugin::GetModularFeatureName());
+
+	// Query all extension plugins to see if we need to use a custom render bridge
+	PFN_xrGetInstanceProcAddr GetProcAddr = nullptr;
+	for (IOpenXRExtensionPlugin* Plugin : ExtModules)
+	{
+		// We are taking ownership of the CustomRenderBridge instance here.
+		TRefCountPtr<FOpenXRRenderBridge> CustomRenderBridge = Plugin->GetCustomRenderBridge(Instance, System);
+		if (CustomRenderBridge)
+		{
+			// We pick the first
+			RenderBridge = CustomRenderBridge;
+			return true;
+		}
+	}
+
+
 	FString RHIString = FApp::GetGraphicsRHI();
 	if (RHIString.IsEmpty())
 	{
@@ -341,34 +388,6 @@ bool FOpenXRHMDPlugin::InitRenderBridge()
 	return true;
 }
 
-struct AnsiKeyFunc : BaseKeyFuncs<const ANSICHAR*, const ANSICHAR*, false>
-{
-	typedef typename TTypeTraits<const ANSICHAR*>::ConstPointerType KeyInitType;
-	typedef typename TCallTraits<const ANSICHAR*>::ParamType ElementInitType;
-
-	/**
-	 * @return The key used to index the given element.
-	 */
-	static FORCEINLINE KeyInitType GetSetKey(ElementInitType Element)
-	{
-		return Element;
-	}
-
-	/**
-	 * @return True if the keys match.
-	 */
-	static FORCEINLINE bool Matches(KeyInitType A, KeyInitType B)
-	{
-		return FCStringAnsi::Strcmp(A, B) == 0;
-	}
-
-	/** Calculates a hash index for a key. */
-	static FORCEINLINE uint32 GetKeyHash(KeyInitType Key)
-	{
-		return GetTypeHash(Key);
-	}
-};
-
 PFN_xrGetInstanceProcAddr FOpenXRHMDPlugin::GetDefaultLoader()
 {
 #if PLATFORM_WINDOWS
@@ -2301,6 +2320,13 @@ void FOpenXRHMD::OnBeginRendering_RenderThread(FRHICommandListImmediate& RHICmdL
 			Layer.bUpdateTexture = Layer.Desc.Flags & IStereoLayers::LAYER_FLAG_TEX_CONTINUOUS_UPDATE;
 		});
 
+		{
+			FReadScopeLock Lock(SessionHandleMutex);
+			for (IOpenXRExtensionPlugin* Module : ExtensionPlugins)
+			{
+				Module->OnAcquireSwapchainImage(Session);
+			}
+		}
 		if (bNeedsAcquireOnRHI)
 		{
 			// TODO: Current spec incorrectly requires this function to have access to the Vulkan queue.
@@ -2440,20 +2466,7 @@ bool FOpenXRHMD::OnStartGameFrame(FWorldContext& WorldContext)
 			CurrentSessionState = SessionState.state;
 
 #if 0
-			static const TCHAR* StateTextMap[] =
-			{
-				TEXT("XR_SESSION_STATE_UNKNOWN"),
-				TEXT("XR_SESSION_STATE_IDLE"),
-				TEXT("XR_SESSION_STATE_READY"),
-				TEXT("XR_SESSION_STATE_SYNCHRONIZED"),
-				TEXT("XR_SESSION_STATE_VISIBLE"),
-				TEXT("XR_SESSION_STATE_FOCUSED"),
-				TEXT("XR_SESSION_STATE_STOPPING"),
-				TEXT("XR_SESSION_STATE_LOSS_PENDING"),
-				TEXT("XR_SESSION_STATE_EXITING"),
-			};
-			TCHAR * StateText = ((int)CurrentSessionState <= (int)XR_SESSION_STATE_EXITING) ? StateTextMap[(int)CurrentSessionState] : TEXT("");
-			UE_LOG(LogHMD, Log, TEXT("Session state switching to %s"), StateText);
+			UE_LOG(LogHMD, Log, TEXT("Session state switching to %s"), OpenXRSessionStateToString(CurrentSessionState));
 #endif
 
 			if (SessionState.state == XR_SESSION_STATE_READY)
diff --git a/Engine/Plugins/Runtime/OpenXR/Source/OpenXRHMD/Private/OpenXRHMD_RenderBridge.h b/Engine/Plugins/Runtime/OpenXR/Source/OpenXRHMD/Private/OpenXRHMD_RenderBridge.h
deleted file mode 100644
diff --git a/Engine/Plugins/Runtime/OpenXR/Source/OpenXRHMD/Private/OpenXRHMD_Swapchain.cpp b/Engine/Plugins/Runtime/OpenXR/Source/OpenXRHMD/Private/OpenXRHMD_Swapchain.cpp
--- a/Engine/Plugins/Runtime/OpenXR/Source/OpenXRHMD/Private/OpenXRHMD_Swapchain.cpp
+++ b/Engine/Plugins/Runtime/OpenXR/Source/OpenXRHMD/Private/OpenXRHMD_Swapchain.cpp
@@ -72,7 +72,7 @@ void FOpenXRSwapchain::ReleaseCurrentImage_RHIThread()
 	XR_ENSURE(xrReleaseSwapchainImage(Handle, &ReleaseInfo));
 }
 
-uint8 GetNearestSupportedSwapchainFormat(XrSession InSession, uint8 RequestedFormat, TFunction<uint32(uint8)> ToPlatformFormat = nullptr)
+uint8 FOpenXRSwapchain::GetNearestSupportedSwapchainFormat(XrSession InSession, uint8 RequestedFormat, TFunction<uint32(uint8)> ToPlatformFormat /*= nullptr*/)
 {
 	if (!ToPlatformFormat)
 	{
@@ -135,7 +135,7 @@ uint8 GetNearestSupportedSwapchainFormat(XrSession InSession, uint8 RequestedFor
 	return FallbackFormat;
 }
 
-XrSwapchain CreateSwapchain(XrSession InSession, uint32 PlatformFormat, uint32 SizeX, uint32 SizeY, uint32 ArraySize, uint32 NumMips, uint32 NumSamples, ETextureCreateFlags Flags, ETextureCreateFlags TargetableTextureFlags)
+XrSwapchain FOpenXRSwapchain::CreateSwapchain(XrSession InSession, uint32 PlatformFormat, uint32 SizeX, uint32 SizeY, uint32 ArraySize, uint32 NumMips, uint32 NumSamples, ETextureCreateFlags Flags, ETextureCreateFlags TargetableTextureFlags)
 {
 	// Need a mutable format so we can reinterpret an sRGB format into a linear format
 	XrSwapchainUsageFlags Usage = XR_SWAPCHAIN_USAGE_MUTABLE_FORMAT_BIT;
@@ -205,13 +205,13 @@ FXRSwapChainPtr CreateSwapchain_D3D11(XrSession InSession, uint8 Format, uint32
 		return PlatformFormat;
 	};
 
-	Format = GetNearestSupportedSwapchainFormat(InSession, Format, ToPlatformFormat);
+	Format = FOpenXRSwapchain::GetNearestSupportedSwapchainFormat(InSession, Format, ToPlatformFormat);
 	if (!Format)
 	{
 		return nullptr;
 	}
 
-	XrSwapchain Swapchain = CreateSwapchain(InSession, ToPlatformFormat(Format), SizeX, SizeY, ArraySize, NumMips, NumSamples, Flags, TargetableTextureFlags);
+	XrSwapchain Swapchain = FOpenXRSwapchain::CreateSwapchain(InSession, ToPlatformFormat(Format), SizeX, SizeY, ArraySize, NumMips, NumSamples, Flags, TargetableTextureFlags);
 	if (!Swapchain)
 	{
 		return nullptr;
@@ -243,13 +243,13 @@ FXRSwapChainPtr CreateSwapchain_D3D12(XrSession InSession, uint8 Format, uint32
 		return PlatformFormat;
 	};
 
-	Format = GetNearestSupportedSwapchainFormat(InSession, Format, ToPlatformFormat);
+	Format = FOpenXRSwapchain::GetNearestSupportedSwapchainFormat(InSession, Format, ToPlatformFormat);
 	if (!Format)
 	{
 		return nullptr;
 	}
 
-	XrSwapchain Swapchain = CreateSwapchain(InSession, ToPlatformFormat(Format), SizeX, SizeY, ArraySize, NumMips, NumSamples, Flags, TargetableTextureFlags);
+	XrSwapchain Swapchain = FOpenXRSwapchain::CreateSwapchain(InSession, ToPlatformFormat(Format), SizeX, SizeY, ArraySize, NumMips, NumSamples, Flags, TargetableTextureFlags);
 	if (!Swapchain)
 	{
 		return nullptr;
@@ -271,13 +271,13 @@ FXRSwapChainPtr CreateSwapchain_D3D12(XrSession InSession, uint8 Format, uint32
 #ifdef XR_USE_GRAPHICS_API_OPENGL
 FXRSwapChainPtr CreateSwapchain_OpenGL(XrSession InSession, uint8 Format, uint32 SizeX, uint32 SizeY, uint32 ArraySize, uint32 NumMips, uint32 NumSamples, ETextureCreateFlags Flags, ETextureCreateFlags TargetableTextureFlags, const FClearValueBinding& ClearValueBinding)
 {
-	Format = GetNearestSupportedSwapchainFormat(InSession, Format);
+	Format = FOpenXRSwapchain::GetNearestSupportedSwapchainFormat(InSession, Format);
 	if (!Format)
 	{
 		return nullptr;
 	}
 
-	XrSwapchain Swapchain = CreateSwapchain(InSession, GPixelFormats[Format].PlatformFormat, SizeX, SizeY, ArraySize, NumMips, NumSamples, Flags, TargetableTextureFlags);
+	XrSwapchain Swapchain = FOpenXRSwapchain::CreateSwapchain(InSession, GPixelFormats[Format].PlatformFormat, SizeX, SizeY, ArraySize, NumMips, NumSamples, Flags, TargetableTextureFlags);
 	if (!Swapchain)
 	{
 		return nullptr;
@@ -304,13 +304,13 @@ FXRSwapChainPtr CreateSwapchain_Vulkan(XrSession InSession, uint8 Format, uint32
 		// UE4 renders a gamma-corrected image so we need to use an sRGB format if available
 		return UEToVkTextureFormat(GPixelFormats[InFormat].UnrealFormat, true);
 	};
-	Format = GetNearestSupportedSwapchainFormat(InSession, Format, ToPlatformFormat);
+	Format = FOpenXRSwapchain::GetNearestSupportedSwapchainFormat(InSession, Format, ToPlatformFormat);
 	if (!Format)
 	{
 		return nullptr;
 	}
 
-	XrSwapchain Swapchain = CreateSwapchain(InSession, ToPlatformFormat(Format), SizeX, SizeY, ArraySize, NumMips, NumSamples, Flags, TargetableTextureFlags);
+	XrSwapchain Swapchain = FOpenXRSwapchain::CreateSwapchain(InSession, ToPlatformFormat(Format), SizeX, SizeY, ArraySize, NumMips, NumSamples, Flags, TargetableTextureFlags);
 	if (!Swapchain)
 	{
 		return nullptr;
diff --git a/Engine/Plugins/Runtime/OpenXR/Source/OpenXRHMD/Private/OpenXRHMD_Swapchain.h b/Engine/Plugins/Runtime/OpenXR/Source/OpenXRHMD/Private/OpenXRHMD_Swapchain.h
deleted file mode 100644
diff --git a/Engine/Plugins/Runtime/OpenXR/Source/OpenXRHMD/Private/OpenXRPlatformRHI.h b/Engine/Plugins/Runtime/OpenXR/Source/OpenXRHMD/Private/OpenXRPlatformRHI.h
deleted file mode 100644
diff --git a/Engine/Plugins/Runtime/OpenXR/Source/OpenXRHMD/Public/IOpenXRExtensionPlugin.h b/Engine/Plugins/Runtime/OpenXR/Source/OpenXRHMD/Public/IOpenXRExtensionPlugin.h
--- a/Engine/Plugins/Runtime/OpenXR/Source/OpenXRHMD/Public/IOpenXRExtensionPlugin.h
+++ b/Engine/Plugins/Runtime/OpenXR/Source/OpenXRHMD/Public/IOpenXRExtensionPlugin.h
@@ -142,6 +142,16 @@ public:
 	{
 		return false;
 	}
+	
+	/**
+	* Optionally provide a custom render bridge for the OpenXR plugin.
+	* Note: this returns a pointer to a new instance allocated with "new".  Calling code is responsible for eventually deleting it.
+	*/
+	virtual class FOpenXRRenderBridge* GetCustomRenderBridge(XrInstance InInstance, XrSystemId InSystem)
+	{
+		return nullptr;
+	}
+
 
 	/**
 	* Fill the array with extensions required by the plugin
@@ -263,7 +273,12 @@ public:
 		return InNext;
 	}
 
-	// OpenXRHMD::OnBeginRendering_RenderThread
+	// OpenXRHMD::OnBeginRendering_RenderThread, before acquiring swapchain
+	virtual void OnAcquireSwapchainImage(XrSession InSession)
+	{
+	}
+
+	// OpenXRHMD::OnBeginRendering_RHIThread
 	virtual const void* OnBeginFrame(XrSession InSession, XrTime DisplayTime, const void* InNext)
 	{
 		return InNext;
diff --git a/Engine/Plugins/Runtime/OpenXR/Source/OpenXRHMD/Public/OpenXRCore.h b/Engine/Plugins/Runtime/OpenXR/Source/OpenXRHMD/Public/OpenXRCore.h
--- a/Engine/Plugins/Runtime/OpenXR/Source/OpenXRHMD/Public/OpenXRCore.h
+++ b/Engine/Plugins/Runtime/OpenXR/Source/OpenXRHMD/Public/OpenXRCore.h
@@ -12,11 +12,21 @@ constexpr const TCHAR* OpenXRResultToString(XrResult e)
 {
 	switch (e)
 	{
-		XR_LIST_ENUM_XrResult(XR_ENUM_CASE_STR)
+		XR_LIST_ENUM_XrResult(XR_ENUM_CASE_STR);
 		default: return TEXT("Unknown");
 	}
 }
 
+#define XR_SESSION_STATE_STR(name, val) case name: return TEXT(#name);
+constexpr const TCHAR* OpenXRSessionStateToString(XrSessionState e)
+{
+	switch (e)
+	{
+		XR_LIST_ENUM_XrSessionState(XR_SESSION_STATE_STR);
+	default: return TEXT("Unknown");
+	}
+}
+
 #if DO_CHECK
 #define XR_ENSURE(x) [] (XrResult Result) \
 	{ \
diff --git a/Engine/Plugins/Runtime/OpenXR/Source/OpenXRHMD/Public/OpenXRHMD_RenderBridge.h b/Engine/Plugins/Runtime/OpenXR/Source/OpenXRHMD/Public/OpenXRHMD_RenderBridge.h
new file mode 100644
--- /dev/null
+++ b/Engine/Plugins/Runtime/OpenXR/Source/OpenXRHMD/Public/OpenXRHMD_RenderBridge.h
@@ -0,0 +1,70 @@
+// Copyright Epic Games, Inc. All Rights Reserved.
+
+#pragma once
+
+#include "CoreMinimal.h"
+#include "HeadMountedDisplayBase.h"
+#include "XRRenderBridge.h"
+#include "XRSwapChain.h"
+#include "OpenXRPlatformRHI.h"
+
+#include <openxr/openxr.h>
+
+class FOpenXRHMD;
+class FRHICommandListImmediate;
+
+class FOpenXRRenderBridge : public FXRRenderBridge
+{
+public:
+	FOpenXRRenderBridge()
+		: AdapterLuid(0)
+		, OpenXRHMD(nullptr)
+	{ }
+
+	void SetOpenXRHMD(FOpenXRHMD* InHMD) { OpenXRHMD = InHMD; }
+	virtual uint64 GetGraphicsAdapterLuid() { return AdapterLuid; }
+
+	virtual void* GetGraphicsBinding() = 0;
+	virtual FXRSwapChainPtr CreateSwapchain(XrSession InSession, uint8 Format, uint32 SizeX, uint32 SizeY, uint32 ArraySize, uint32 NumMips, uint32 NumSamples, ETextureCreateFlags Flags, ETextureCreateFlags TargetableTextureFlags, const FClearValueBinding& ClearValueBinding) = 0;
+
+	FXRSwapChainPtr CreateSwapchain(XrSession InSession, FRHITexture2D* Template, ETextureCreateFlags Flags, ETextureCreateFlags TargetableTextureFlags)
+	{
+		if (!Template)
+		{
+			return nullptr;
+		}
+
+		return CreateSwapchain(InSession,
+			Template->GetFormat(),
+			Template->GetSizeX(),
+			Template->GetSizeY(),
+			1,
+			Template->GetNumMips(),
+			Template->GetNumSamples(),
+			Template->GetFlags() | Flags,
+			TargetableTextureFlags,
+			Template->GetClearBinding());
+	}
+
+	/** FRHICustomPresent */
+	virtual bool Present(int32& InOutSyncInterval) override;
+
+protected:
+	uint64 AdapterLuid;
+
+private:
+	FOpenXRHMD* OpenXRHMD;
+};
+
+#ifdef XR_USE_GRAPHICS_API_D3D11
+FOpenXRRenderBridge* CreateRenderBridge_D3D11(XrInstance InInstance, XrSystemId InSystem);
+#endif
+#ifdef XR_USE_GRAPHICS_API_D3D12
+FOpenXRRenderBridge* CreateRenderBridge_D3D12(XrInstance InInstance, XrSystemId InSystem);
+#endif
+#ifdef XR_USE_GRAPHICS_API_OPENGL
+FOpenXRRenderBridge* CreateRenderBridge_OpenGL(XrInstance InInstance, XrSystemId InSystem);
+#endif
+#ifdef XR_USE_GRAPHICS_API_VULKAN
+FOpenXRRenderBridge* CreateRenderBridge_Vulkan(XrInstance InInstance, XrSystemId InSystem);
+#endif
diff --git a/Engine/Plugins/Runtime/OpenXR/Source/OpenXRHMD/Public/OpenXRHMD_Swapchain.h b/Engine/Plugins/Runtime/OpenXR/Source/OpenXRHMD/Public/OpenXRHMD_Swapchain.h
new file mode 100644
--- /dev/null
+++ b/Engine/Plugins/Runtime/OpenXR/Source/OpenXRHMD/Public/OpenXRHMD_Swapchain.h
@@ -0,0 +1,41 @@
+// Copyright Epic Games, Inc. All Rights Reserved.
+
+#pragma once
+
+#include "CoreMinimal.h"
+#include "HeadMountedDisplayBase.h"
+#include "XRSwapChain.h"
+#include "OpenXRPlatformRHI.h"
+
+#include <openxr/openxr.h>
+
+class FOpenXRSwapchain : public FXRSwapChain
+{
+public:
+	FOpenXRSwapchain(TArray<FTextureRHIRef>&& InRHITextureSwapChain, const FTextureRHIRef & InRHITexture, XrSwapchain InHandle);
+	virtual ~FOpenXRSwapchain();
+
+	virtual void IncrementSwapChainIndex_RHIThread() override final;
+	virtual void WaitCurrentImage_RHIThread(int64 Timeout) override final;
+	virtual void ReleaseCurrentImage_RHIThread() override final;
+
+	XrSwapchain GetHandle() { return Handle; }
+	static XrSwapchain CreateSwapchain(XrSession InSession, uint32 PlatformFormat, uint32 SizeX, uint32 SizeY, uint32 ArraySize, uint32 NumMips, uint32 NumSamples, ETextureCreateFlags Flags, ETextureCreateFlags TargetableTextureFlags);
+	static uint8 GetNearestSupportedSwapchainFormat(XrSession InSession, uint8 RequestedFormat, TFunction<uint32(uint8)> ToPlatformFormat = nullptr);
+
+protected:
+	XrSwapchain Handle;
+};
+
+#ifdef XR_USE_GRAPHICS_API_D3D11
+FXRSwapChainPtr CreateSwapchain_D3D11(XrSession InSession, uint8 Format, uint32 SizeX, uint32 SizeY, uint32 ArraySize, uint32 NumMips, uint32 NumSamples, ETextureCreateFlags Flags, ETextureCreateFlags TargetableTextureFlags, const FClearValueBinding& ClearValueBinding);
+#endif
+#ifdef XR_USE_GRAPHICS_API_D3D12
+FXRSwapChainPtr CreateSwapchain_D3D12(XrSession InSession, uint8 Format, uint32 SizeX, uint32 SizeY, uint32 ArraySize, uint32 NumMips, uint32 NumSamples, ETextureCreateFlags Flags, ETextureCreateFlags TargetableTextureFlags, const FClearValueBinding& ClearValueBinding);
+#endif
+#ifdef XR_USE_GRAPHICS_API_OPENGL
+FXRSwapChainPtr CreateSwapchain_OpenGL(XrSession InSession, uint8 Format, uint32 SizeX, uint32 SizeY, uint32 ArraySize, uint32 NumMips, uint32 NumSamples, ETextureCreateFlags Flags, ETextureCreateFlags TargetableTextureFlags, const FClearValueBinding& ClearValueBinding);
+#endif
+#ifdef XR_USE_GRAPHICS_API_VULKAN
+FXRSwapChainPtr CreateSwapchain_Vulkan(XrSession InSession, uint8 Format, uint32 SizeX, uint32 SizeY, uint32 ArraySize, uint32 NumMips, uint32 NumSamples, ETextureCreateFlags Flags, ETextureCreateFlags TargetableTextureFlags, const FClearValueBinding& ClearValueBinding);
+#endif
\ No newline at end of file
diff --git a/Engine/Plugins/Runtime/OpenXR/Source/OpenXRHMD/Public/OpenXRPlatformRHI.h b/Engine/Plugins/Runtime/OpenXR/Source/OpenXRHMD/Public/OpenXRPlatformRHI.h
new file mode 100644
--- /dev/null
+++ b/Engine/Plugins/Runtime/OpenXR/Source/OpenXRHMD/Public/OpenXRPlatformRHI.h
@@ -0,0 +1,106 @@
+// Copyright Epic Games, Inc. All Rights Reserved.
+
+#pragma once
+
+#if PLATFORM_WINDOWS || PLATFORM_HOLOLENS
+#define XR_USE_PLATFORM_WIN32		1
+#define XR_USE_GRAPHICS_API_D3D11	1
+#define XR_USE_GRAPHICS_API_D3D12	1
+#endif
+
+#if PLATFORM_WINDOWS
+#define XR_USE_GRAPHICS_API_OPENGL	1
+#endif
+
+#if PLATFORM_ANDROID
+#define XR_USE_PLATFORM_ANDROID 1
+#endif
+
+#if PLATFORM_WINDOWS || PLATFORM_ANDROID || PLATFORM_LINUX
+#define XR_USE_GRAPHICS_API_VULKAN 1
+#endif
+
+//-------------------------------------------------------------------------------------------------
+// D3D11
+//-------------------------------------------------------------------------------------------------
+
+#ifdef XR_USE_GRAPHICS_API_D3D11
+#include "D3D11RHIPrivate.h"
+#include "D3D11Util.h"
+#endif // XR_USE_GRAPHICS_API_D3D11
+
+
+//-------------------------------------------------------------------------------------------------
+// D3D12
+//-------------------------------------------------------------------------------------------------
+
+#ifdef XR_USE_GRAPHICS_API_D3D12
+#define GetD3D11CubeFace GetD3D12CubeFace
+#define VerifyD3D11Result VerifyD3D12Result
+#define GetD3D11TextureFromRHITexture GetD3D12TextureFromRHITexture
+#define FRingAllocation FRingAllocation_D3D12
+#define GetRenderTargetFormat GetRenderTargetFormat_D3D12
+#define ED3D11ShaderOffsetBuffer ED3D12ShaderOffsetBuffer
+#define FindShaderResourceDXGIFormat FindShaderResourceDXGIFormat_D3D12
+#define FindUnorderedAccessDXGIFormat FindUnorderedAccessDXGIFormat_D3D12
+#define FindDepthStencilDXGIFormat FindDepthStencilDXGIFormat_D3D12
+#define HasStencilBits HasStencilBits_D3D12
+#define FVector4VertexDeclaration FVector4VertexDeclaration_D3D12
+#define GLOBAL_CONSTANT_BUFFER_INDEX GLOBAL_CONSTANT_BUFFER_INDEX_D3D12
+#define MAX_CONSTANT_BUFFER_SLOTS MAX_CONSTANT_BUFFER_SLOTS_D3D12
+#define FD3DGPUProfiler FD3D12GPUProfiler
+#define FRangeAllocator FRangeAllocator_D3D12
+
+#include "D3D12RHIPrivate.h"
+#include "D3D12Util.h"
+
+#undef GetD3D11CubeFace
+#undef VerifyD3D11Result
+#undef GetD3D11TextureFromRHITexture
+#undef FRingAllocation
+#undef GetRenderTargetFormat
+#undef ED3D11ShaderOffsetBuffer
+#undef FindShaderResourceDXGIFormat
+#undef FindUnorderedAccessDXGIFormat
+#undef FindDepthStencilDXGIFormat
+#undef HasStencilBits
+#undef FVector4VertexDeclaration
+#undef GLOBAL_CONSTANT_BUFFER_INDEX
+#undef MAX_CONSTANT_BUFFER_SLOTS
+#undef FD3DGPUProfiler
+#undef FRangeAllocator
+#endif // XR_USE_GRAPHICS_API_D3D12
+
+
+//-------------------------------------------------------------------------------------------------
+// OpenGL
+//-------------------------------------------------------------------------------------------------
+
+#ifdef XR_USE_GRAPHICS_API_OPENGL
+#include "OpenGLDrvPrivate.h"
+#include "OpenGLResources.h"
+#endif // XR_USE_GRAPHICS_API_OPENGL
+
+
+//-------------------------------------------------------------------------------------------------
+// Vulkan
+//-------------------------------------------------------------------------------------------------
+
+#ifdef XR_USE_GRAPHICS_API_VULKAN
+#include "VulkanRHIPrivate.h"
+#include "VulkanResources.h"
+#endif // XR_USE_GRAPHICS_API_VULKAN
+
+#if PLATFORM_WINDOWS
+#include "Windows/AllowWindowsPlatformTypes.h"
+#endif
+
+#if PLATFORM_ANDROID
+#include <android_native_app_glue.h>
+#endif
+
+#include <openxr/openxr_platform.h>
+
+#if PLATFORM_WINDOWS
+#include "Windows/HideWindowsPlatformTypes.h"
+#endif
diff --git a/Engine/Plugins/Runtime/OpenXR/Source/OpenXRInput/Private/OpenXRInput.h b/Engine/Plugins/Runtime/OpenXR/Source/OpenXRInput/Private/OpenXRInput.h
--- a/Engine/Plugins/Runtime/OpenXR/Source/OpenXRInput/Private/OpenXRInput.h
+++ b/Engine/Plugins/Runtime/OpenXR/Source/OpenXRInput/Private/OpenXRInput.h
@@ -13,13 +13,11 @@ class FOpenXRHMD;
 struct FInputActionKeyMapping;
 struct FInputAxisKeyMapping;
 
-#if PLATFORM_LINUX
-// When compiling on Linux the XrPath type becomes ambigious for overloading
+// On some platforms the XrPath type becomes ambigious for overloading
 FORCEINLINE uint32 GetTypeHash(const TPair<XrPath, XrPath>& Pair)
 {
 	return HashCombine(GetTypeHash((uint64)Pair.Key), GetTypeHash((uint64)Pair.Value));
 }
-#endif
 
 class FOpenXRInputPlugin : public IOpenXRInputPlugin
 {
