Remove win32

diff --git a/Engine/Plugins/Runtime/OpenXR/Source/OpenXRHMD/OpenXRHMD.Build.cs b/Engine/Plugins/Runtime/OpenXR/Source/OpenXRHMD/OpenXRHMD.Build.cs
--- a/Engine/Plugins/Runtime/OpenXR/Source/OpenXRHMD/OpenXRHMD.Build.cs
+++ b/Engine/Plugins/Runtime/OpenXR/Source/OpenXRHMD/OpenXRHMD.Build.cs
@@ -98,7 +98,7 @@ namespace UnrealBuildTool.Rules
                 AddEngineThirdPartyPrivateStaticDependencies(Target, "OpenGL");
 			}
 
-			if (Target.Platform == UnrealTargetPlatform.Win32 || Target.Platform == UnrealTargetPlatform.Win64 || Target.Platform == UnrealTargetPlatform.Android)
+			if (Target.Platform == UnrealTargetPlatform.Win64 || Target.Platform == UnrealTargetPlatform.Android)
             {
                 PrivateDependencyModuleNames.AddRange(new string[] {
                     "VulkanRHI"
