Fix OpenXRHMD compile errors on Linux
#rnx

diff --git a/Engine/Plugins/Runtime/OpenXR/Source/OpenXRHMD/OpenXRHMD.Build.cs b/Engine/Plugins/Runtime/OpenXR/Source/OpenXRHMD/OpenXRHMD.Build.cs
--- a/Engine/Plugins/Runtime/OpenXR/Source/OpenXRHMD/OpenXRHMD.Build.cs
+++ b/Engine/Plugins/Runtime/OpenXR/Source/OpenXRHMD/OpenXRHMD.Build.cs
@@ -98,7 +98,8 @@ namespace UnrealBuildTool.Rules
                 AddEngineThirdPartyPrivateStaticDependencies(Target, "OpenGL");
 			}
 
-			if (Target.Platform == UnrealTargetPlatform.Win64 || Target.Platform == UnrealTargetPlatform.Android)
+			if (Target.Platform == UnrealTargetPlatform.Win64 || Target.Platform == UnrealTargetPlatform.Android  
+			    || Target.IsInPlatformGroup(UnrealPlatformGroup.Linux))
             {
                 PrivateDependencyModuleNames.AddRange(new string[] {
                     "VulkanRHI"
