OpenXR related features and fixes.
-Adds support for getting velocity and accelleration through GetControllerTransformForTime.
-Adds mechanism for using openxr chain structs which are defined in and consumed by an extension plugin.
-Support for defining a swapchain and renderbridge in an extension plugin.
-Some OpenXR fixes related to runtimes not supplying tracking data of various types.  Basically if the data isn't being provided we should not copy it out because it might be full of nonsense.  Also we would generally rather preserve our last cached position than overwrite with zero/identity.
-Moving OpenXRPlatformRHI.h OpenXRHMD_RenderBridge.h and _Swapchain.h back to private folder.
-Previously one could not run the same haptic effect on two devices simultaneously and have it behave correctly because most of the per-effect-play data for each effect was shared by all controllers.  This refactors so that the necessary runtime data is per controller and the constant bulk data is shared by all controllers.
-Adding infrastructure for playing haptics on a VR HMD.
-Refactored haptic data resampling into the oculus plugin because that is plugin specific.
swarm reviews review-18220544 and review-18220542

diff --git a/Engine/Plugins/Animation/LiveLink/Source/LiveLink/Private/LiveLinkMotionController.h b/Engine/Plugins/Animation/LiveLink/Source/LiveLink/Private/LiveLinkMotionController.h
--- a/Engine/Plugins/Animation/LiveLink/Source/LiveLink/Private/LiveLinkMotionController.h
+++ b/Engine/Plugins/Animation/LiveLink/Source/LiveLink/Private/LiveLinkMotionController.h
@@ -70,11 +70,12 @@ public:
 		return false;
 	}
 
-	virtual bool GetControllerOrientationAndPositionForTime(const int32 ControllerIndex, const FName MotionSource, FTimespan Time, bool& OutTimeWasUsed, FRotator& OutOrientation, FVector& OutPosition, bool& OutbProvidedLinearVelocity, FVector& OutLinearVelocity, bool& OutbProvidedAngularVelocity, FVector& OutAngularVelocityRadPerSec, float WorldToMetersScale) const override
+	virtual bool GetControllerOrientationAndPositionForTime(const int32 ControllerIndex, const FName MotionSource, FTimespan Time, bool& OutTimeWasUsed, FRotator& OutOrientation, FVector& OutPosition, bool& OutbProvidedLinearVelocity, FVector& OutLinearVelocity, bool& OutbProvidedAngularVelocity, FVector& OutAngularVelocityRadPerSec, bool& OutbProvidedLinearAcceleration, FVector& OutLinearAcceleration, float WorldToMetersScale) const override
 	{
 		OutTimeWasUsed = false;
 		OutbProvidedLinearVelocity = false;
 		OutbProvidedAngularVelocity = false;
+		OutbProvidedLinearAcceleration = false;
 		return GetControllerOrientationAndPosition(ControllerIndex, MotionSource, OutOrientation, OutPosition, WorldToMetersScale);
 	}
 
diff --git a/Engine/Plugins/Runtime/AR/Google/GoogleARCore/Source/GoogleARCoreBase/Private/GoogleARCoreMotionController.h b/Engine/Plugins/Runtime/AR/Google/GoogleARCore/Source/GoogleARCoreBase/Private/GoogleARCoreMotionController.h
--- a/Engine/Plugins/Runtime/AR/Google/GoogleARCore/Source/GoogleARCoreBase/Private/GoogleARCoreMotionController.h
+++ b/Engine/Plugins/Runtime/AR/Google/GoogleARCore/Source/GoogleARCoreBase/Private/GoogleARCoreMotionController.h
@@ -37,11 +37,12 @@ public:
 	*/
 	virtual bool GetControllerOrientationAndPosition(const int32 ControllerIndex, const FName MotionSource, FRotator& OutOrientation, FVector& OutPosition, float WorldToMetersScale) const override;
 
-	virtual bool GetControllerOrientationAndPositionForTime(const int32 ControllerIndex, const FName MotionSource, FTimespan Time, bool& OutTimeWasUsed, FRotator& OutOrientation, FVector& OutPosition, bool& OutbProvidedLinearVelocity, FVector& OutLinearVelocity, bool& OutbProvidedAngularVelocity, FVector& OutAngularVelocityRadPerSec, float WorldToMetersScale) const override
+	virtual bool GetControllerOrientationAndPositionForTime(const int32 ControllerIndex, const FName MotionSource, FTimespan Time, bool& OutTimeWasUsed, FRotator& OutOrientation, FVector& OutPosition, bool& OutbProvidedLinearVelocity, FVector& OutLinearVelocity, bool& OutbProvidedAngularVelocity, FVector& OutAngularVelocityRadPerSec, bool& OutbProvidedLinearAcceleration, FVector& OutLinearAcceleration, float WorldToMetersScale) const override
 	{
 		OutTimeWasUsed = false;
 		OutbProvidedLinearVelocity = false;
 		OutbProvidedAngularVelocity = false;
+		OutbProvidedLinearAcceleration = false;
 		return GetControllerOrientationAndPosition(ControllerIndex, MotionSource, OutOrientation, OutPosition, WorldToMetersScale);
 	}
 
diff --git a/Engine/Plugins/Runtime/Oculus/OculusVR/Source/OculusInput/Private/OculusInput.cpp b/Engine/Plugins/Runtime/Oculus/OculusVR/Source/OculusInput/Private/OculusInput.cpp
--- a/Engine/Plugins/Runtime/Oculus/OculusVR/Source/OculusInput/Private/OculusInput.cpp
+++ b/Engine/Plugins/Runtime/Oculus/OculusVR/Source/OculusInput/Private/OculusInput.cpp
@@ -1106,9 +1106,13 @@ void FOculusInput::SetHapticFeedbackValues(int32 ControllerId, int32 Hand, const
 						(OvrpControllerState.ConnectedControllerTypes & (ovrpController_Touch | ovrpController_LTrackedRemote | ovrpController_RTrackedRemote)))
 					{
 						// Buffered haptics is currently only supported on Touch
-						FHapticFeedbackBuffer* HapticBuffer = Values.HapticBuffer;
+						if (Values.HapticBuffer)
+						{
+							ControllerState.ResampleHapticBufferData(*Values.HapticBuffer, ResampledRawDataCache);
+						}
+						FHapticFeedbackBuffer* HapticBuffer = &ControllerState.ResampledHapticBuffer;
 						if ( (OvrpControllerState.ConnectedControllerTypes & (ovrpController_Touch)) &&
-							HapticBuffer && HapticBuffer->SamplingRate == OvrpHapticsDesc.SampleRateHz)
+							Values.HapticBuffer && HapticBuffer->SamplingRate == OvrpHapticsDesc.SampleRateHz)
 						{
 							const ovrpController OvrpController = (EControllerHand(Hand) == EControllerHand::Left) ? ovrpController_LTouch : ovrpController_RTouch;
 
@@ -1131,7 +1135,8 @@ void FOculusInput::SetHapticFeedbackValues(int32 ControllerId, int32 Hand, const
 
 									if (OvrpHapticsBuffer.SamplesCount == 0 && OvrpHapticsState.SamplesQueued == 0)
 									{
-										HapticBuffer->bFinishedPlaying = true;
+										Values.HapticBuffer->bFinishedPlaying = HapticBuffer->bFinishedPlaying = true;
+
 										ControllerState.bPlayingHapticEffect = false;
 									}
 									else
@@ -1151,7 +1156,7 @@ void FOculusInput::SetHapticFeedbackValues(int32 ControllerId, int32 Hand, const
 											for (int i = 0; i < OvrpHapticsBuffer.SamplesCount; i++)
 											{
 												const uint32 DataIndex = HapticBuffer->CurrentPtr + (i * 2);
-												const uint16* const RawData = reinterpret_cast<uint16*>(&HapticBuffer->RawData[DataIndex]);
+												const uint16* const RawData = reinterpret_cast<const uint16*>(&HapticBuffer->RawData[DataIndex]);
 												samples[i] = static_cast<uint16>(*RawData * HapticBuffer->ScaleFactor);
 											}
 											OvrpHapticsBuffer.Samples = bufferToFree = samples;
@@ -1162,7 +1167,7 @@ void FOculusInput::SetHapticFeedbackValues(int32 ControllerId, int32 Hand, const
 											for (int i = 0; i < OvrpHapticsBuffer.SamplesCount; i++)
 											{
 												const uint32 DataIndex = HapticBuffer->CurrentPtr + (i * 4);
-												const uint32* const RawData = reinterpret_cast<uint32*>(&HapticBuffer->RawData[DataIndex]);
+												const uint32* const RawData = reinterpret_cast<const uint32*>(&HapticBuffer->RawData[DataIndex]);
 												samples[i] = static_cast<uint32>(*RawData * HapticBuffer->ScaleFactor);
 											}
 											OvrpHapticsBuffer.Samples = bufferToFree = samples;
@@ -1227,6 +1232,62 @@ void FOculusInput::SetHapticFeedbackValues(int32 ControllerId, int32 Hand, const
 	}
 }
 
+void FOculusTouchControllerState::ResampleHapticBufferData(const FHapticFeedbackBuffer& HapticBuffer, TMap<const uint8*, TSharedPtr<TArray<uint8>>>& ResampledRawDataCache)
+{
+	const uint8* OriginalRawData = HapticBuffer.RawData;
+	TSharedPtr<TArray<uint8>>* ResampledRawDataSharedPtrPtr = ResampledRawDataCache.Find(OriginalRawData);
+	if (ResampledRawDataSharedPtrPtr == nullptr)
+	{
+		// We need to resample and cache the resampled data.
+
+		ResampledHapticBuffer = HapticBuffer;
+
+		int32 SampleRate = HapticBuffer.SamplingRate;
+		int TargetFrequency = 320;
+		int TargetBufferSize = (HapticBuffer.BufferLength * TargetFrequency) / (SampleRate * 2) + 1; //2 because we're only using half of the 16bit source PCM buffer
+		ResampledHapticBuffer.BufferLength = TargetBufferSize;
+		ResampledHapticBuffer.CurrentPtr = 0;
+		ResampledHapticBuffer.SamplingRate = TargetFrequency;
+
+		TSharedPtr<TArray<uint8>>& NewResampledRawDataSharedPtr = ResampledRawDataCache.Add(OriginalRawData);
+		NewResampledRawDataSharedPtr = MakeShared<TArray<uint8>>();
+		ResampledRawDataSharedPtrPtr = &NewResampledRawDataSharedPtr;
+		TArray<uint8>& ResampledRawData = *NewResampledRawDataSharedPtr;
+		ResampledRawData.SetNum(TargetBufferSize);
+
+		const uint8* PCMData = HapticBuffer.RawData;
+
+		int previousTargetIndex = -1;
+		int currentMin = 0;
+		for (int i = 1; i < HapticBuffer.BufferLength; i += 2)
+		{
+			int targetIndex = i * TargetFrequency / (SampleRate * 2);
+			int val = PCMData[i];
+			if (val & 0x80)
+			{
+				val = ~val;
+			}
+			currentMin = FMath::Min(currentMin, val);
+
+			if (targetIndex != previousTargetIndex)
+			{
+
+				ResampledRawData[targetIndex] = val * 2;// *Scale;
+				previousTargetIndex = targetIndex;
+				currentMin = 0;
+			}
+		}
+
+		ResampledHapticBuffer.RawData = ResampledRawData.GetData();
+	}
+	else if (ResampledHapticBuffer.RawData != (*ResampledRawDataSharedPtrPtr)->GetData())
+	{
+		// If this a cached effect, but not the same one we played last so we need to copy the new one's buffer and reference its cached resampled data.
+		ResampledHapticBuffer = HapticBuffer;
+		ResampledHapticBuffer.RawData = (*ResampledRawDataSharedPtrPtr)->GetData();
+	}
+}
+
 void FOculusInput::GetHapticFrequencyRange(float& MinFrequency, float& MaxFrequency) const
 {
 	MinFrequency = 0.f;
diff --git a/Engine/Plugins/Runtime/Oculus/OculusVR/Source/OculusInput/Private/OculusInput.h b/Engine/Plugins/Runtime/Oculus/OculusVR/Source/OculusInput/Private/OculusInput.h
--- a/Engine/Plugins/Runtime/Oculus/OculusVR/Source/OculusInput/Private/OculusInput.h
+++ b/Engine/Plugins/Runtime/Oculus/OculusVR/Source/OculusInput/Private/OculusInput.h
@@ -103,6 +103,10 @@ private:
 	ovrpHapticsDesc OvrpHapticsDesc;
 
 	int LocalTrackingSpaceRecenterCount;
+
+	// Maintain a cache of resampled raw data so we don't resample it on every play.  This is a map of OriginalRawData pointers, used only as a key, to ResampledRawData buffers.
+	// The values are pointers because the map could be reallocated and we cache raw pointers to the uint8 array data elsewhere.
+	TMap<const uint8*, TSharedPtr<TArray<uint8>>> ResampledRawDataCache;
 };
 
 
diff --git a/Engine/Plugins/Runtime/Oculus/OculusVR/Source/OculusInput/Private/OculusInputState.h b/Engine/Plugins/Runtime/Oculus/OculusVR/Source/OculusInput/Private/OculusInputState.h
--- a/Engine/Plugins/Runtime/Oculus/OculusVR/Source/OculusInput/Private/OculusInputState.h
+++ b/Engine/Plugins/Runtime/Oculus/OculusVR/Source/OculusInput/Private/OculusInputState.h
@@ -343,7 +343,10 @@ struct FOculusTouchControllerState : FOculusControllerState
 	/** Number of times that controller was recentered (for mobile controllers) */
 	int RecenterCount;
 
-
+public:
+	FHapticFeedbackBuffer ResampledHapticBuffer;
+	void ResampleHapticBufferData(const FHapticFeedbackBuffer& HapticBuffer, TMap<const uint8*, TSharedPtr<TArray<uint8>>>& ResampledRawDataCache);
+	
 	/** Explicit constructor sets up sensible defaults */
 	FOculusTouchControllerState( const EControllerHand Hand )
 		: TriggerAxis( 0.0f ),
diff --git a/Engine/Plugins/Runtime/OpenXR/Source/OpenXRHMD/Private/Epic_openxr.h b/Engine/Plugins/Runtime/OpenXR/Source/OpenXRHMD/Private/Epic_openxr.h
new file mode 100644
--- /dev/null
+++ b/Engine/Plugins/Runtime/OpenXR/Source/OpenXRHMD/Private/Epic_openxr.h
@@ -0,0 +1,37 @@
+// Copyright Epic Games, Inc. All Rights Reserved.
+
+#pragma once
+
+#include <openxr/openxr.h>
+
+// Uncomment to make intellisense work better.
+//#include "/../../../../../Source/ThirdParty/OpenXR/include/openxr/openxr.h"
+
+// This number is in a block that is unused by openxr, and that seemed unlikely to be used soon based on the patterns of usage.
+// But if it ever was used by openxr we would have to change it.
+#define XR_TYPE_EPIC 1000059050
+
+typedef enum XrStructureTypeEPIC {
+	XR_TYPE_SPACE_ACCELERATION_EPIC = XR_TYPE_EPIC + 0,
+} XrStructureTypeEPIC;
+
+
+//  Extensions
+
+#define XR_EPIC_space_acceleration_1
+#define XR_EPIC_space_acceleration_SPEC_VERSION 1
+#define XR_EPIC_SPACE_ACCELERATION_NAME "XR_EPIC_space_acceleration"
+
+typedef XrFlags64 XrSpaceAccelerationFlagsEPIC;
+
+// Flag bits for XrSpaceAccelerationFlagsEPIC
+static const XrSpaceAccelerationFlagsEPIC XR_SPACE_ACCELERATION_LINEAR_VALID_BIT_EPIC = 0x00000001;
+static const XrSpaceAccelerationFlagsEPIC XR_SPACE_ACCELERATION_ANGULAR_VALID_BIT_EPIC = 0x00000002;
+
+typedef struct XrSpaceAccelerationEPIC {
+	XrStructureType							type;
+	void* XR_MAY_ALIAS						next;
+	XrSpaceAccelerationFlagsEPIC			accelerationFlags;
+	XrVector3f								linearAcceleration;  // m/s^2
+	XrVector3f							angularAcceleration; // rad/s^2
+} XrSpaceAccelerationEPIC;
\ No newline at end of file
diff --git a/Engine/Plugins/Runtime/OpenXR/Source/OpenXRHMD/Private/OpenXRHMD.cpp b/Engine/Plugins/Runtime/OpenXR/Source/OpenXRHMD/Private/OpenXRHMD.cpp
--- a/Engine/Plugins/Runtime/OpenXR/Source/OpenXRHMD/Private/OpenXRHMD.cpp
+++ b/Engine/Plugins/Runtime/OpenXR/Source/OpenXRHMD/Private/OpenXRHMD.cpp
@@ -28,6 +28,7 @@
 #include "IHandTracker.h"
 #include "PixelShaderUtils.h"
 #include "ScenePrivate.h"
+#include "Epic_openxr.h"
 
 #if WITH_EDITOR
 #include "Editor/UnrealEd/Classes/Editor/EditorEngine.h"
@@ -424,7 +425,7 @@ bool FOpenXRHMD::GetCurrentPose(int32 DeviceId, FQuat& CurrentOrientation, FVect
 	return true;
 }
 
-bool FOpenXRHMD::GetPoseForTime(int32 DeviceId, FTimespan Timespan, FQuat& Orientation, FVector& Position, bool& bProvidedLinearVelocity, FVector& LinearVelocity, bool& bProvidedAngularVelocity, FVector& AngularVelocityRadPerSec)
+bool FOpenXRHMD::GetPoseForTime(int32 DeviceId, FTimespan Timespan, bool& OutTimeWasUsed, FQuat& Orientation, FVector& Position, bool& bProvidedLinearVelocity, FVector& LinearVelocity, bool& bProvidedAngularVelocity, FVector& AngularVelocityRadPerSec, bool& bProvidedLinearAcceleration, FVector& LinearAcceleration, float InWorldToMetersScale)
 {
 	FPipelinedFrameState& PipelineState = GetPipelinedFrameStateForThread();
 
@@ -436,34 +437,54 @@ bool FOpenXRHMD::GetPoseForTime(int32 DeviceId, FTimespan Timespan, FQuat& Orien
 
 	XrTime TargetTime = ToXrTime(Timespan);
 
+	// If TargetTime is zero just get the latest data (rather than the oldest).
+	if (TargetTime == 0)
+	{
+		OutTimeWasUsed = false;
+		TargetTime = GetDisplayTime();
+	}
+	else
+	{
+		OutTimeWasUsed = true;
+	}
+
 	const FDeviceSpace& DeviceSpace = DeviceSpaces[DeviceId];
 
-	XrSpaceVelocity DeviceVelocity { XR_TYPE_SPACE_VELOCITY };
+	XrSpaceAccelerationEPIC DeviceAcceleration{ (XrStructureType)XR_TYPE_SPACE_ACCELERATION_EPIC };
+	XrSpaceVelocity DeviceVelocity { XR_TYPE_SPACE_VELOCITY, &DeviceAcceleration };
 	XrSpaceLocation DeviceLocation { XR_TYPE_SPACE_LOCATION, &DeviceVelocity };
 
 	XR_ENSURE(xrLocateSpace(DeviceSpace.Space, PipelineState.TrackingSpace, TargetTime, &DeviceLocation));
 
+	bool ReturnValue = false;
+
 	if (DeviceLocation.locationFlags & XR_SPACE_LOCATION_ORIENTATION_VALID_BIT &&
 		DeviceLocation.locationFlags & XR_SPACE_LOCATION_POSITION_VALID_BIT)
 	{
 		Orientation = ToFQuat(DeviceLocation.pose.orientation);
-		Position = ToFVector(DeviceLocation.pose.position, GetWorldToMetersScale());
+		Position = ToFVector(DeviceLocation.pose.position, InWorldToMetersScale);
 
 		if (DeviceVelocity.velocityFlags & XR_SPACE_VELOCITY_LINEAR_VALID_BIT)
 		{
 			bProvidedLinearVelocity = true;
-			LinearVelocity = ToFVector(DeviceVelocity.linearVelocity, GetWorldToMetersScale());
+			LinearVelocity = ToFVector(DeviceVelocity.linearVelocity, InWorldToMetersScale);
 		}
 		if (DeviceVelocity.velocityFlags & XR_SPACE_VELOCITY_ANGULAR_VALID_BIT)
 		{
 			bProvidedAngularVelocity = true;
-			AngularVelocityRadPerSec = ToFVector(DeviceVelocity.angularVelocity);
+			AngularVelocityRadPerSec = -ToFVector(DeviceVelocity.angularVelocity);
 		}
 
-		return true;
+		if (DeviceAcceleration.accelerationFlags & XR_SPACE_ACCELERATION_LINEAR_VALID_BIT_EPIC)
+		{
+			bProvidedLinearAcceleration = true;
+			LinearAcceleration = ToFVector(DeviceAcceleration.linearAcceleration, InWorldToMetersScale);
+		}
+
+		ReturnValue = true;
 	}
 
-	return false;
+	return ReturnValue;
 }
 
 bool FOpenXRHMD::IsChromaAbCorrectionEnabled() const
@@ -978,6 +999,11 @@ FOpenXRHMD::FOpenXRHMD(const FAutoRegister& AutoRegister, XrInstance InInstance,
 
 	// Give the all frame states the same initial values.
 	PipelinedFrameStateRHI = PipelinedFrameStateRendering = PipelinedFrameStateGame;
+
+	for (IOpenXRExtensionPlugin* Module : ExtensionPlugins)
+	{
+		Module->BindExtensionPluginDelegates(*this);
+	}
 }
 
 FOpenXRHMD::~FOpenXRHMD()
@@ -1033,11 +1059,9 @@ void FOpenXRHMD::UpdateDeviceLocations(bool bUpdateOpenXRExtensionPlugins)
 			const FDeviceSpace& DeviceSpace = DeviceSpaces[DeviceIndex];
 			if (DeviceSpace.Space != XR_NULL_HANDLE)
 			{
-				XrSpaceLocation& DeviceLocation = PipelineState.DeviceLocations[DeviceIndex];
-				DeviceLocation.type = XR_TYPE_SPACE_LOCATION;
-				DeviceLocation.next = nullptr;
-
-				XrResult Result = xrLocateSpace(DeviceSpace.Space, PipelineState.TrackingSpace, PipelineState.FrameState.predictedDisplayTime, &DeviceLocation);
+				XrSpaceLocation NewDeviceLocation = {};
+				NewDeviceLocation.type = XR_TYPE_SPACE_LOCATION;
+				XrResult Result = xrLocateSpace(DeviceSpace.Space, PipelineState.TrackingSpace, PipelineState.FrameState.predictedDisplayTime, &NewDeviceLocation);
 				if (Result == XR_ERROR_TIME_INVALID)
 				{
 					// The display time is no longer valid so set the location as invalid as well
@@ -1047,6 +1071,20 @@ void FOpenXRHMD::UpdateDeviceLocations(bool bUpdateOpenXRExtensionPlugins)
 				{
 					XR_ENSURE(Result);
 				}
+				
+				XrSpaceLocation& CachedDeviceLocation = PipelineState.DeviceLocations[DeviceIndex];
+				// Clear the location tracked bits
+				CachedDeviceLocation.locationFlags &= ~(XR_SPACE_LOCATION_POSITION_TRACKED_BIT | XR_SPACE_LOCATION_ORIENTATION_TRACKED_BIT);
+				if (NewDeviceLocation.locationFlags & (XR_SPACE_LOCATION_POSITION_VALID_BIT))
+				{
+					CachedDeviceLocation.pose.position = NewDeviceLocation.pose.position;
+					CachedDeviceLocation.locationFlags |= (NewDeviceLocation.locationFlags & (XR_SPACE_LOCATION_POSITION_TRACKED_BIT | XR_SPACE_LOCATION_POSITION_VALID_BIT));
+				}
+				if (NewDeviceLocation.locationFlags & (XR_SPACE_LOCATION_ORIENTATION_VALID_BIT))
+				{
+					CachedDeviceLocation.pose.orientation = NewDeviceLocation.pose.orientation;
+					CachedDeviceLocation.locationFlags |= (NewDeviceLocation.locationFlags & (XR_SPACE_LOCATION_ORIENTATION_TRACKED_BIT | XR_SPACE_LOCATION_ORIENTATION_VALID_BIT));
+				}
 			}
 			else
 			{
@@ -1805,6 +1843,13 @@ void FOpenXRHMD::OnBeginRendering_RenderThread(FRHICommandListImmediate& RHICmdL
 
 		SCOPED_NAMED_EVENT(EnqueueFrame, FColor::Red);
 
+		{
+			FReadScopeLock Lock(SessionHandleMutex);
+			for (IOpenXRExtensionPlugin* Module : ExtensionPlugins)
+			{
+				Module->OnAcquireSwapchainImage(Session);
+			}
+		}
 		// Reset the update flag on all layers
 		ForEachLayer([&](uint32 /* unused */, FOpenXRLayer& Layer)
 		{
@@ -2184,6 +2229,7 @@ void FOpenXRHMD::OnFinishRendering_RHIThread()
 		EndInfo.next = nullptr;
 		EndInfo.displayTime = PipelinedFrameStateRHI.FrameState.predictedDisplayTime;
 		EndInfo.environmentBlendMode = SelectedEnvironmentBlendMode;
+
 		EndInfo.layerCount = PipelinedFrameStateRHI.FrameState.shouldRender ? Headers.Num() : 0;
 		EndInfo.layers = PipelinedFrameStateRHI.FrameState.shouldRender ? Headers.GetData() : nullptr;
 
@@ -2277,6 +2323,8 @@ void FOpenXRHMD::CopyTexture_RenderThread(FRHICommandListImmediate& RHICmdList,
 		VSize = SrcRect.Height() / SrcTextureHeight;
 	}
 
+	RHICmdList.Transition(FRHITransitionInfo(DstTexture, ERHIAccess::Unknown, ERHIAccess::RTV));
+
 	FRHITexture * ColorRT = DstTexture->GetTexture2DArray() ? DstTexture->GetTexture2DArray() : DstTexture->GetTexture2D();
 	FRHIRenderPassInfo RenderPassInfo(ColorRT, RTAction);
 	RHICmdList.BeginRenderPass(RenderPassInfo, TEXT("OpenXRHMD_CopyTexture"));
@@ -2331,8 +2379,16 @@ void FOpenXRHMD::CopyTexture_RenderThread(FRHICommandListImmediate& RHICmdList,
 			FIntPoint(1, 1),
 			VertexShader,
 			EDRF_Default);
+
 	}
 	RHICmdList.EndRenderPass();
+	
+	// TODO
+	// We can't assume the new access state is valid for usage by clients. For the spectator screen, Present is fine.
+	// For layer blits, we'll need to have something that's compatible with swapchain requirements when handed off to
+	// compositor (RTV?). There's a pending refactor from Jules that will make it easier for caller to pass in
+	// desired ERHIAccess for destination texture.
+	RHICmdList.Transition(FRHITransitionInfo(DstTexture, ERHIAccess::RTV, (ERHIAccess::Present | ERHIAccess::SRVMask)));
 }
 
 void FOpenXRHMD::CopyTexture_RenderThread(FRHICommandListImmediate& RHICmdList, FRHITexture2D* SrcTexture, FIntRect SrcRect, const FXRSwapChainPtr& DstSwapChain, FIntRect DstRect, bool bClearBlack, bool bNoAlpha) const
@@ -2344,7 +2400,7 @@ void FOpenXRHMD::CopyTexture_RenderThread(FRHICommandListImmediate& RHICmdList,
 	});
 
 	// Now that we've enqueued the swapchain wait we can add the commands to do the actual texture copy
-	FRHITexture2DArray* DstTexture = DstSwapChain->GetTexture2DArray();
+	FRHITexture2D* const DstTexture = DstSwapChain->GetTexture2DArray() ? DstSwapChain->GetTexture2DArray() : DstSwapChain->GetTexture2D();
 	CopyTexture_RenderThread(RHICmdList, SrcTexture, SrcRect, DstTexture, DstRect, bClearBlack, bNoAlpha, ERenderTargetActions::Clear_Store);
 
 	// Enqueue a command to release the image after the copy is done
diff --git a/Engine/Plugins/Runtime/OpenXR/Source/OpenXRHMD/Private/OpenXRHMD.h b/Engine/Plugins/Runtime/OpenXR/Source/OpenXRHMD/Private/OpenXRHMD.h
--- a/Engine/Plugins/Runtime/OpenXR/Source/OpenXRHMD/Private/OpenXRHMD.h
+++ b/Engine/Plugins/Runtime/OpenXR/Source/OpenXRHMD/Private/OpenXRHMD.h
@@ -14,6 +14,7 @@
 #include "StereoLayerManager.h"
 #include "DefaultSpectatorScreenController.h"
 #include "IHeadMountedDisplayVulkanExtensions.h"
+#include "IOpenXRExtensionPluginDelegates.h"
 
 #include <openxr/openxr.h>
 
@@ -32,6 +33,7 @@ class FOpenXRHMD
 	, public FHMDSceneViewExtension
 	, public FOpenXRAssetManager
 	, public TStereoLayerManager<FOpenXRLayer>
+	, public IOpenXRExtensionPluginDelegates
 {
 public:
 	class FDeviceSpace
@@ -133,7 +135,7 @@ public:
 
 	virtual bool GetIsTracked(int32 DeviceId);
 	virtual bool GetCurrentPose(int32 DeviceId, FQuat& CurrentOrientation, FVector& CurrentPosition) override;
-	virtual bool GetPoseForTime(int32 DeviceId, FTimespan Timespan, FQuat& CurrentOrientation, FVector& CurrentPosition, bool& bProvidedLinearVelocity, FVector& LinearVelocity, bool& bProvidedAngularVelocity, FVector& AngularVelocityRadPerSec);
+	virtual bool GetPoseForTime(int32 DeviceId, FTimespan Timespan, bool& OutTimeWasUsed, FQuat& CurrentOrientation, FVector& CurrentPosition, bool& bProvidedLinearVelocity, FVector& LinearVelocity, bool& bProvidedAngularVelocity, FVector& AngularVelocityRadPerSec, bool& bProvidedLinearAcceleration, FVector& LinearAcceleration, float WorldToMetersScale);
 	virtual void SetBaseRotation(const FRotator& BaseRot) override;
 	virtual FRotator GetBaseRotation() const override;
 
@@ -264,6 +266,12 @@ public:
 	/** IStereoLayers */
 	virtual bool ShouldCopyDebugLayersToSpectatorScreen() const override { return true; }
 
+	/** IOpenXRExtensionPluginDelegates */
+public:
+	virtual FApplyHapticFeedbackAddChainStructsDelegate& GetApplyHapticFeedbackAddChainStructsDelegate() override { return ApplyHapticFeedbackAddChainStructsDelegate; }
+private:
+	FApplyHapticFeedbackAddChainStructsDelegate ApplyHapticFeedbackAddChainStructsDelegate;
+
 public:
 	/** Constructor */
 	FOpenXRHMD(const FAutoRegister&, XrInstance InInstance, XrSystemId InSystem, TRefCountPtr<FOpenXRRenderBridge>& InRenderBridge, TArray<const char*> InEnabledExtensions, TArray<class IOpenXRExtensionPlugin*> InExtensionPlugins, IARSystemSupport* ARSystemSupport);
diff --git a/Engine/Plugins/Runtime/OpenXR/Source/OpenXRHMD/Private/OpenXRHMD_RenderBridge.h b/Engine/Plugins/Runtime/OpenXR/Source/OpenXRHMD/Private/OpenXRHMD_RenderBridge.h
new file mode 100644
--- /dev/null
+++ b/Engine/Plugins/Runtime/OpenXR/Source/OpenXRHMD/Private/OpenXRHMD_RenderBridge.h
@@ -0,0 +1,70 @@
+// Copyright Epic Games, Inc. All Rights Reserved.
+
+#pragma once
+
+#include "CoreMinimal.h"
+#include "HeadMountedDisplayBase.h"
+#include "XRRenderBridge.h"
+#include "XRSwapChain.h"
+#include "OpenXRPlatformRHI.h"
+
+#include <openxr/openxr.h>
+
+class FOpenXRHMD;
+class FRHICommandListImmediate;
+
+class FOpenXRRenderBridge : public FXRRenderBridge
+{
+public:
+	FOpenXRRenderBridge()
+		: AdapterLuid(0)
+		, OpenXRHMD(nullptr)
+	{ }
+
+	void SetOpenXRHMD(FOpenXRHMD* InHMD) { OpenXRHMD = InHMD; }
+	virtual uint64 GetGraphicsAdapterLuid() { return AdapterLuid; }
+
+	virtual void* GetGraphicsBinding() = 0;
+	virtual FXRSwapChainPtr CreateSwapchain(XrSession InSession, uint8 Format, uint32 SizeX, uint32 SizeY, uint32 ArraySize, uint32 NumMips, uint32 NumSamples, ETextureCreateFlags Flags, ETextureCreateFlags TargetableTextureFlags, const FClearValueBinding& ClearValueBinding) = 0;
+
+	FXRSwapChainPtr CreateSwapchain(XrSession InSession, FRHITexture2D* Template, ETextureCreateFlags Flags, ETextureCreateFlags TargetableTextureFlags)
+	{
+		if (!Template)
+		{
+			return nullptr;
+		}
+
+		return CreateSwapchain(InSession,
+			Template->GetFormat(),
+			Template->GetSizeX(),
+			Template->GetSizeY(),
+			1,
+			Template->GetNumMips(),
+			Template->GetNumSamples(),
+			Template->GetFlags() | Flags,
+			TargetableTextureFlags,
+			Template->GetClearBinding());
+	}
+
+	/** FRHICustomPresent */
+	virtual bool Present(int32& InOutSyncInterval) override;
+
+protected:
+	uint64 AdapterLuid;
+
+private:
+	FOpenXRHMD* OpenXRHMD;
+};
+
+#ifdef XR_USE_GRAPHICS_API_D3D11
+FOpenXRRenderBridge* CreateRenderBridge_D3D11(XrInstance InInstance, XrSystemId InSystem);
+#endif
+#ifdef XR_USE_GRAPHICS_API_D3D12
+FOpenXRRenderBridge* CreateRenderBridge_D3D12(XrInstance InInstance, XrSystemId InSystem);
+#endif
+#ifdef XR_USE_GRAPHICS_API_OPENGL
+FOpenXRRenderBridge* CreateRenderBridge_OpenGL(XrInstance InInstance, XrSystemId InSystem);
+#endif
+#ifdef XR_USE_GRAPHICS_API_VULKAN
+FOpenXRRenderBridge* CreateRenderBridge_Vulkan(XrInstance InInstance, XrSystemId InSystem);
+#endif
diff --git a/Engine/Plugins/Runtime/OpenXR/Source/OpenXRHMD/Private/OpenXRHMD_Swapchain.h b/Engine/Plugins/Runtime/OpenXR/Source/OpenXRHMD/Private/OpenXRHMD_Swapchain.h
new file mode 100644
--- /dev/null
+++ b/Engine/Plugins/Runtime/OpenXR/Source/OpenXRHMD/Private/OpenXRHMD_Swapchain.h
@@ -0,0 +1,45 @@
+// Copyright Epic Games, Inc. All Rights Reserved.
+
+#pragma once
+
+#include "CoreMinimal.h"
+#include "HeadMountedDisplayBase.h"
+#include "XRSwapChain.h"
+#include "OpenXRPlatformRHI.h"
+
+#include <openxr/openxr.h>
+
+class FOpenXRSwapchain : public FXRSwapChain
+{
+public:
+	FOpenXRSwapchain(TArray<FTextureRHIRef>&& InRHITextureSwapChain, const FTextureRHIRef & InRHITexture, XrSwapchain InHandle);
+	virtual ~FOpenXRSwapchain();
+
+	virtual void IncrementSwapChainIndex_RHIThread() override final;
+	virtual void WaitCurrentImage_RHIThread(int64 Timeout) override final;
+	virtual void ReleaseCurrentImage_RHIThread() override final;
+
+	XrSwapchain GetHandle() { return Handle; }
+	static XrSwapchain CreateSwapchain(XrSession InSession, uint32 PlatformFormat, uint32 SizeX, uint32 SizeY, uint32 ArraySize, uint32 NumMips, uint32 NumSamples, ETextureCreateFlags Flags, ETextureCreateFlags TargetableTextureFlags);
+	static uint8 GetNearestSupportedSwapchainFormat(XrSession InSession, uint8 RequestedFormat, TFunction<uint32(uint8)> ToPlatformFormat = nullptr);
+
+protected:
+	XrSwapchain Handle;
+	std::atomic<bool> Acquired;
+};
+
+#ifdef XR_USE_GRAPHICS_API_D3D11
+FXRSwapChainPtr CreateSwapchain_D3D11(XrSession InSession, uint8 Format, uint32 SizeX, uint32 SizeY, uint32 ArraySize, uint32 NumMips, uint32 NumSamples, ETextureCreateFlags Flags, ETextureCreateFlags TargetableTextureFlags, const FClearValueBinding& ClearValueBinding);
+#endif
+#ifdef XR_USE_GRAPHICS_API_D3D12
+FXRSwapChainPtr CreateSwapchain_D3D12(XrSession InSession, uint8 Format, uint32 SizeX, uint32 SizeY, uint32 ArraySize, uint32 NumMips, uint32 NumSamples, ETextureCreateFlags Flags, ETextureCreateFlags TargetableTextureFlags, const FClearValueBinding& ClearValueBinding);
+#endif
+#ifdef XR_USE_GRAPHICS_API_OPENGL
+FXRSwapChainPtr CreateSwapchain_OpenGL(XrSession InSession, uint8 Format, uint32 SizeX, uint32 SizeY, uint32 ArraySize, uint32 NumMips, uint32 NumSamples, ETextureCreateFlags Flags, ETextureCreateFlags TargetableTextureFlags, const FClearValueBinding& ClearValueBinding);
+#endif
+#ifdef XR_USE_GRAPHICS_API_OPENGL_ES
+FXRSwapChainPtr CreateSwapchain_OpenGLES(XrSession InSession, uint8 Format, uint32 SizeX, uint32 SizeY, uint32 ArraySize, uint32 NumMips, uint32 NumSamples, ETextureCreateFlags Flags, ETextureCreateFlags TargetableTextureFlags, const FClearValueBinding& ClearValueBinding);
+#endif
+#ifdef XR_USE_GRAPHICS_API_VULKAN
+FXRSwapChainPtr CreateSwapchain_Vulkan(XrSession InSession, uint8 Format, uint32 SizeX, uint32 SizeY, uint32 ArraySize, uint32 NumMips, uint32 NumSamples, ETextureCreateFlags Flags, ETextureCreateFlags TargetableTextureFlags, const FClearValueBinding& ClearValueBinding);
+#endif
\ No newline at end of file
diff --git a/Engine/Plugins/Runtime/OpenXR/Source/OpenXRHMD/Private/OpenXRPlatformRHI.h b/Engine/Plugins/Runtime/OpenXR/Source/OpenXRHMD/Private/OpenXRPlatformRHI.h
new file mode 100644
--- /dev/null
+++ b/Engine/Plugins/Runtime/OpenXR/Source/OpenXRHMD/Private/OpenXRPlatformRHI.h
@@ -0,0 +1,96 @@
+// Copyright Epic Games, Inc. All Rights Reserved.
+
+#pragma once
+
+#if PLATFORM_WINDOWS || PLATFORM_HOLOLENS
+#define XR_USE_PLATFORM_WIN32		1
+#define XR_USE_GRAPHICS_API_D3D11	1
+#define XR_USE_GRAPHICS_API_D3D12	1
+#endif
+
+#if PLATFORM_WINDOWS
+#define XR_USE_GRAPHICS_API_OPENGL	1
+#endif
+
+#if PLATFORM_ANDROID
+#define XR_USE_PLATFORM_ANDROID 1
+#define XR_USE_GRAPHICS_API_OPENGL_ES	1
+#endif
+
+#if PLATFORM_WINDOWS || PLATFORM_ANDROID || PLATFORM_LINUX
+#define XR_USE_GRAPHICS_API_VULKAN 1
+#endif
+
+//-------------------------------------------------------------------------------------------------
+// D3D11
+//-------------------------------------------------------------------------------------------------
+
+#ifdef XR_USE_GRAPHICS_API_D3D11
+#include "D3D11RHIPrivate.h"
+#include "D3D11Util.h"
+#endif // XR_USE_GRAPHICS_API_D3D11
+
+
+//-------------------------------------------------------------------------------------------------
+// D3D12
+//-------------------------------------------------------------------------------------------------
+
+#ifdef XR_USE_GRAPHICS_API_D3D12
+#include "D3D12RHIPrivate.h"
+#include "D3D12Util.h"
+
+#undef GetD3D11CubeFace
+#undef VerifyD3D11Result
+#undef GetD3D11TextureFromRHITexture
+#undef FRingAllocation
+#undef GetRenderTargetFormat
+#undef ED3D11ShaderOffsetBuffer
+#undef FindShaderResourceDXGIFormat
+#undef FindUnorderedAccessDXGIFormat
+#undef FindDepthStencilDXGIFormat
+#undef HasStencilBits
+#undef FVector4VertexDeclaration
+#undef GLOBAL_CONSTANT_BUFFER_INDEX
+#undef MAX_CONSTANT_BUFFER_SLOTS
+#undef FD3DGPUProfiler
+#undef FRangeAllocator
+#endif // XR_USE_GRAPHICS_API_D3D12
+
+
+//-------------------------------------------------------------------------------------------------
+// OpenGL
+//-------------------------------------------------------------------------------------------------
+
+#if defined(XR_USE_GRAPHICS_API_OPENGL) || defined(XR_USE_GRAPHICS_API_OPENGL_ES)
+#include "OpenGLDrvPrivate.h"
+#include "OpenGLResources.h"
+
+#ifdef XR_USE_PLATFORM_ANDROID
+#include "Android/AndroidEGL.h"
+#include "Android/AndroidOpenGL.h"
+#endif
+
+#endif // XR_USE_GRAPHICS_API_OPENGL
+
+//-------------------------------------------------------------------------------------------------
+// Vulkan
+//-------------------------------------------------------------------------------------------------
+
+#ifdef XR_USE_GRAPHICS_API_VULKAN
+#include "VulkanRHIPrivate.h"
+#include "VulkanResources.h"
+#endif // XR_USE_GRAPHICS_API_VULKAN
+
+#if PLATFORM_WINDOWS
+#include "Windows/AllowWindowsPlatformTypes.h"
+#endif
+
+#if PLATFORM_ANDROID
+#include <android_native_app_glue.h>
+#endif
+
+#include <openxr/openxr_platform.h>
+
+#if PLATFORM_WINDOWS
+#include "Windows/HideWindowsPlatformTypes.h"
+#endif
diff --git a/Engine/Plugins/Runtime/OpenXR/Source/OpenXRHMD/Public/IOpenXRExtensionPlugin.h b/Engine/Plugins/Runtime/OpenXR/Source/OpenXRHMD/Public/IOpenXRExtensionPlugin.h
--- a/Engine/Plugins/Runtime/OpenXR/Source/OpenXRHMD/Public/IOpenXRExtensionPlugin.h
+++ b/Engine/Plugins/Runtime/OpenXR/Source/OpenXRHMD/Public/IOpenXRExtensionPlugin.h
@@ -275,6 +275,10 @@ public:
 	{
 	}
 
+	virtual void BindExtensionPluginDelegates(class IOpenXRExtensionPluginDelegates& OpenXRHMD)
+	{
+	}
+
 	virtual const void* OnGetSystem(XrInstance InInstance, const void* InNext)
 	{
 		return InNext;
diff --git a/Engine/Plugins/Runtime/OpenXR/Source/OpenXRHMD/Public/IOpenXRExtensionPluginDelegates.h b/Engine/Plugins/Runtime/OpenXR/Source/OpenXRHMD/Public/IOpenXRExtensionPluginDelegates.h
new file mode 100644
--- /dev/null
+++ b/Engine/Plugins/Runtime/OpenXR/Source/OpenXRHMD/Public/IOpenXRExtensionPluginDelegates.h
@@ -0,0 +1,57 @@
+// Copyright Epic Games, Inc. All Rights Reserved.
+
+#pragma once
+
+#include "CoreMinimal.h"
+
+struct XrSpaceLocation;
+
+// When adding chain structs from an ExtensionPlugin it is convenient to have the calling code take ownership of the chain struct so that it can be deleted after the openxr call.
+// We can also handle the linked list stuff here.
+// 
+// SAMPLE_USAGE in calling code (FOpenXRHMD::GetPoseForTime):
+//		FOpenXRExtensionChainStructPtrs ScopedExtensionChainStructs; // Own the chain structs until we are done with them.
+//		GetPoseForTimeAddChainStructsDelegate.Broadcast(&DeviceLocation, ScopedExtensionChainStructs);
+//
+// SAMPLE USAGE in Delegate, SpaceLocation would be the struct that is the head of the chain and we are adding an XrSpaceVelocity struct to the chain:
+// void FFireAntSpaceVelocityModule::GetPoseForTimeAddChainStructs(XrSpaceLocation* SpaceLocation, FOpenXRExtensionChainStructPtrs& ChainStructPtrs)
+// {
+//		ChainStructPtrs.Add(MakeShared<TOpenXRExtensionChainStruct<XrSpaceVelocity>>(SpaceLocation, XR_TYPE_SPACE_VELOCITY));
+// }
+// OR, if you need to set data in the struct:
+// void FFireAntSpaceVelocityModule::GetPoseForTimeAddChainStructs(XrSpaceLocation* SpaceLocation, FOpenXRExtensionChainStructPtrs& ChainStructPtrs)
+// {
+//		TSharedPtr<TOpenXRExtensionChainStruct<XrSpaceVelocity>> ChainStructPtr = MakeShared<TOpenXRExtensionChainStruct<XrSpaceVelocity>>(SpaceLocation, XR_TYPE_SPACE_VELOCITY);
+//		ChainStructPtr->Data.linearVelocity = {1, 2, 3}; // Note: this is useless for XrSpaceVelocity because it is an output chain struct.
+//		ChainStructPtrs.Add(ChainStructPtr);
+// }
+
+struct FOpenXRExtensionChainStruct
+{
+	virtual ~FOpenXRExtensionChainStruct() {};
+};
+typedef TSharedPtr<FOpenXRExtensionChainStruct> FOpenXRExtensionChainStructPtr;
+typedef TArray<FOpenXRExtensionChainStructPtr> FOpenXRExtensionChainStructPtrs;
+template <typename ChainStructType>
+struct TOpenXRExtensionChainStruct : public FOpenXRExtensionChainStruct
+{
+	TOpenXRExtensionChainStruct(void* Head, XrStructureType XRType)
+	{
+		XrBaseOutStructure* HeadStructure = (XrBaseOutStructure*)Head;
+		Data.type = XRType;
+		Data.next = HeadStructure->next;
+		HeadStructure->next = (XrBaseOutStructure*)&Data;
+	}
+	virtual ~TOpenXRExtensionChainStruct() {};
+	ChainStructType Data;
+};
+
+// Delegates
+DECLARE_MULTICAST_DELEGATE_ThreeParams(FApplyHapticFeedbackAddChainStructsDelegate, XrHapticVibration*, FOpenXRExtensionChainStructPtrs&, struct FHapticFeedbackBuffer*);
+
+class IOpenXRExtensionPluginDelegates
+{
+public:
+	virtual FApplyHapticFeedbackAddChainStructsDelegate& GetApplyHapticFeedbackAddChainStructsDelegate() = 0;
+};
+
diff --git a/Engine/Plugins/Runtime/OpenXR/Source/OpenXRHMD/Public/OpenXRHMD_RenderBridge.h b/Engine/Plugins/Runtime/OpenXR/Source/OpenXRHMD/Public/OpenXRHMD_RenderBridge.h
deleted file mode 100644
diff --git a/Engine/Plugins/Runtime/OpenXR/Source/OpenXRHMD/Public/OpenXRHMD_Swapchain.h b/Engine/Plugins/Runtime/OpenXR/Source/OpenXRHMD/Public/OpenXRHMD_Swapchain.h
deleted file mode 100644
diff --git a/Engine/Plugins/Runtime/OpenXR/Source/OpenXRHMD/Public/OpenXRPlatformRHI.h b/Engine/Plugins/Runtime/OpenXR/Source/OpenXRHMD/Public/OpenXRPlatformRHI.h
deleted file mode 100644
diff --git a/Engine/Plugins/Runtime/OpenXR/Source/OpenXRInput/Private/OpenXRInput.cpp b/Engine/Plugins/Runtime/OpenXR/Source/OpenXRInput/Private/OpenXRInput.cpp
--- a/Engine/Plugins/Runtime/OpenXR/Source/OpenXRInput/Private/OpenXRInput.cpp
+++ b/Engine/Plugins/Runtime/OpenXR/Source/OpenXRInput/Private/OpenXRInput.cpp
@@ -194,7 +194,12 @@ FOpenXRInputPlugin::FOpenXRInput::FOpenXRInput(FOpenXRHMD* HMD)
 
 XrAction FOpenXRInputPlugin::FOpenXRInput::GetActionForMotionSource(FName MotionSource) const
 {
-	const FOpenXRController& Controller = Controllers[MotionSourceToControllerHandMap.FindChecked(MotionSource)];
+	const EControllerHand* Hand = MotionSourceToControllerHandMap.Find(MotionSource);
+	if (Hand == nullptr)
+	{
+		return XR_NULL_HANDLE;
+	}
+	const FOpenXRController& Controller = Controllers[*Hand];
 	if (MotionSource == OpenXRSourceNames::LeftAim || MotionSource == OpenXRSourceNames::RightAim)
 	{
 		return Controller.AimAction;
@@ -259,11 +264,13 @@ void FOpenXRInputPlugin::FOpenXRInput::BuildActions()
 
 	XrPath LeftHand = GetPath(Instance, "/user/hand/left");
 	XrPath RightHand = GetPath(Instance, "/user/hand/right");
+	XrPath HMD = GetPath(Instance, "/user/head");
 
 	// Controller poses
 	OpenXRHMD->ResetActionDevices();
 	Controllers.Add(EControllerHand::Left, FOpenXRController(ActionSet, LeftHand, "Left Controller"));
 	Controllers.Add(EControllerHand::Right, FOpenXRController(ActionSet, RightHand, "Right Controller"));
+	Controllers.Add(EControllerHand::HMD, FOpenXRController(ActionSet, HMD, "HMD"));
 
 	// Generate a map of all supported interaction profiles
 	XrPath SimpleControllerPath = GetPath(Instance, "/interaction_profiles/khr/simple_controller");
@@ -480,6 +487,14 @@ int32 FOpenXRInputPlugin::FOpenXRInput::SuggestBindings(XrInstance Instance, FOp
 					Path += "/click";
 				}
 			}
+			else if (Component == "touch")
+			{
+				Path += "/touch";
+			}
+			else if (Component == "touchaxis")
+			{
+				Path += "/touchvalue";  // Note: this is not a standard openxr identifier.  It is meant to represent some kind of analog touch sensor.
+			}
 			else if (Component == "up" || Component == "down" || Component == "left" || Component == "right")
 			{
 				// TODO: Reserved for D-Pad support
@@ -573,6 +588,11 @@ void FOpenXRInputPlugin::FOpenXRInput::SendControllerEvents()
 		return;
 	}
 
+	if (!OpenXRHMD->IsFocused())
+	{
+		return;
+	}
+
 	XrSession Session = OpenXRHMD->GetSession();
 
 	for (XrPath Subaction : SubactionPaths)
@@ -765,10 +785,8 @@ bool FOpenXRInputPlugin::FOpenXRInput::GetControllerOrientationAndPosition(const
 	return false;
 }
 
-bool FOpenXRInputPlugin::FOpenXRInput::GetControllerOrientationAndPositionForTime(const int32 ControllerIndex, const FName MotionSource, FTimespan Time, bool& OutTimeWasUsed, FRotator& OutOrientation, FVector& OutPosition, bool& OutbProvidedLinearVelocity, FVector& OutLinearVelocity, bool& OutbProvidedAngularVelocity, FVector& OutAngularVelocityRadPerSec, float WorldToMetersScale) const
+bool FOpenXRInputPlugin::FOpenXRInput::GetControllerOrientationAndPositionForTime(const int32 ControllerIndex, const FName MotionSource, FTimespan Time, bool& OutTimeWasUsed, FRotator& OutOrientation, FVector& OutPosition, bool& OutbProvidedLinearVelocity, FVector& OutLinearVelocity, bool& OutbProvidedAngularVelocity, FVector& OutAngularVelocityRadPerSec, bool& OutbProvidedLinearAcceleration, FVector& OutLinearAcceleration, float WorldToMetersScale) const
 {
-	OutTimeWasUsed = true;
-
 	if (OpenXRHMD == nullptr)
 	{
 		return false;
@@ -778,18 +796,18 @@ bool FOpenXRInputPlugin::FOpenXRInput::GetControllerOrientationAndPositionForTim
 	{
 		if (MotionSource == OpenXRSourceNames::AnyHand)
 		{
-			return GetControllerOrientationAndPositionForTime(ControllerIndex, OpenXRSourceNames::LeftGrip, Time, OutTimeWasUsed, OutOrientation, OutPosition, OutbProvidedLinearVelocity, OutLinearVelocity, OutbProvidedAngularVelocity, OutAngularVelocityRadPerSec, WorldToMetersScale)
-				|| GetControllerOrientationAndPositionForTime(ControllerIndex, OpenXRSourceNames::RightGrip, Time, OutTimeWasUsed, OutOrientation, OutPosition, OutbProvidedLinearVelocity, OutLinearVelocity, OutbProvidedAngularVelocity, OutAngularVelocityRadPerSec, WorldToMetersScale);
+			return GetControllerOrientationAndPositionForTime(ControllerIndex, OpenXRSourceNames::LeftGrip, Time, OutTimeWasUsed, OutOrientation, OutPosition, OutbProvidedLinearVelocity, OutLinearVelocity, OutbProvidedAngularVelocity, OutAngularVelocityRadPerSec, OutbProvidedLinearAcceleration, OutLinearAcceleration, WorldToMetersScale)
+				|| GetControllerOrientationAndPositionForTime(ControllerIndex, OpenXRSourceNames::RightGrip, Time, OutTimeWasUsed, OutOrientation, OutPosition, OutbProvidedLinearVelocity, OutLinearVelocity, OutbProvidedAngularVelocity, OutAngularVelocityRadPerSec, OutbProvidedLinearAcceleration, OutLinearAcceleration, WorldToMetersScale);
 		}
 
 		if (MotionSource == OpenXRSourceNames::Left)
 		{
-			return GetControllerOrientationAndPositionForTime(ControllerIndex, OpenXRSourceNames::LeftGrip, Time, OutTimeWasUsed, OutOrientation, OutPosition, OutbProvidedLinearVelocity, OutLinearVelocity, OutbProvidedAngularVelocity, OutAngularVelocityRadPerSec, WorldToMetersScale);
+			return GetControllerOrientationAndPositionForTime(ControllerIndex, OpenXRSourceNames::LeftGrip, Time, OutTimeWasUsed, OutOrientation, OutPosition, OutbProvidedLinearVelocity, OutLinearVelocity, OutbProvidedAngularVelocity, OutAngularVelocityRadPerSec, OutbProvidedLinearAcceleration, OutLinearAcceleration, WorldToMetersScale);
 		}
 
 		if (MotionSource == OpenXRSourceNames::Right)
 		{
-			return GetControllerOrientationAndPositionForTime(ControllerIndex, OpenXRSourceNames::RightGrip, Time, OutTimeWasUsed, OutOrientation, OutPosition, OutbProvidedLinearVelocity, OutLinearVelocity, OutbProvidedAngularVelocity, OutAngularVelocityRadPerSec, WorldToMetersScale);
+			return GetControllerOrientationAndPositionForTime(ControllerIndex, OpenXRSourceNames::RightGrip, Time, OutTimeWasUsed, OutOrientation, OutPosition, OutbProvidedLinearVelocity, OutLinearVelocity, OutbProvidedAngularVelocity, OutAngularVelocityRadPerSec, OutbProvidedLinearAcceleration, OutLinearAcceleration, WorldToMetersScale);
 		}
 	}
 
@@ -806,6 +824,13 @@ bool FOpenXRInputPlugin::FOpenXRInput::GetControllerOrientationAndPositionForTim
 	GetInfo.subactionPath = XR_NULL_PATH;
 	GetInfo.action = GetActionForMotionSource(MotionSource);
 
+	if (GetInfo.action == XR_NULL_HANDLE)
+	{
+		UE_LOG(LogHMD, Warning, TEXT("GetControllerOrientationAndPositionForTime called with motion source %s which is unknown.  Cannot get pose."), *MotionSource.ToString());
+
+		return false;
+	}
+
 	XrActionStatePose State;
 	State.type = XR_TYPE_ACTION_STATE_POSE;
 	State.next = nullptr;
@@ -813,7 +838,7 @@ bool FOpenXRInputPlugin::FOpenXRInput::GetControllerOrientationAndPositionForTim
 	if (Result >= XR_SUCCESS && State.isActive)
 	{
 		FQuat Orientation;
-		OpenXRHMD->GetPoseForTime(GetDeviceIDForMotionSource(MotionSource), Time, Orientation, OutPosition, OutbProvidedLinearVelocity, OutLinearVelocity, OutbProvidedAngularVelocity, OutAngularVelocityRadPerSec);
+		OpenXRHMD->GetPoseForTime(GetDeviceIDForMotionSource(MotionSource), Time, OutTimeWasUsed, Orientation, OutPosition, OutbProvidedLinearVelocity, OutLinearVelocity, OutbProvidedAngularVelocity, OutAngularVelocityRadPerSec, OutbProvidedLinearAcceleration, OutLinearAcceleration, WorldToMetersScale);
 		OutOrientation = FRotator(Orientation);
 		return true;
 	}
@@ -907,6 +932,11 @@ void FOpenXRInputPlugin::FOpenXRInput::SetHapticFeedbackValues(int32 ControllerI
 	{
 		return;
 	}
+	
+	if (!OpenXRHMD->IsFocused())
+	{
+		return;
+	}
 
 	XrHapticVibration HapticValue;
 	HapticValue.type = XR_TYPE_HAPTIC_VIBRATION;
@@ -924,12 +954,17 @@ void FOpenXRInputPlugin::FOpenXRInput::SetHapticFeedbackValues(int32 ControllerI
 			HapticActionInfo.next = nullptr;
 			HapticActionInfo.subactionPath = XR_NULL_PATH;
 			HapticActionInfo.action = Controllers[EControllerHand::Left].VibrationAction;
-			if (Values.Amplitude <= 0.0f || Values.Frequency < XR_FREQUENCY_UNSPECIFIED)
+			if (Values.HapticBuffer == nullptr && (Values.Amplitude <= 0.0f || Values.Frequency < XR_FREQUENCY_UNSPECIFIED))
 			{
 				XR_ENSURE(xrStopHapticFeedback(Session, &HapticActionInfo));
 			}
 			else
-			{
+			{				
+				FOpenXRExtensionChainStructPtrs ScopedExtensionChainStructs;
+				if (Values.HapticBuffer != nullptr)
+				{
+					OpenXRHMD->GetApplyHapticFeedbackAddChainStructsDelegate().Broadcast(&HapticValue, ScopedExtensionChainStructs, Values.HapticBuffer);
+				}
 				XR_ENSURE(xrApplyHapticFeedback(Session, &HapticActionInfo, (const XrHapticBaseHeader*)&HapticValue));
 			}
 		}
@@ -940,7 +975,29 @@ void FOpenXRInputPlugin::FOpenXRInput::SetHapticFeedbackValues(int32 ControllerI
 			HapticActionInfo.next = nullptr;
 			HapticActionInfo.subactionPath = XR_NULL_PATH;
 			HapticActionInfo.action = Controllers[EControllerHand::Right].VibrationAction;
-			if (Values.Amplitude <= 0.0f || Values.Frequency < XR_FREQUENCY_UNSPECIFIED)
+			if (Values.HapticBuffer == nullptr && (Values.Amplitude <= 0.0f || Values.Frequency < XR_FREQUENCY_UNSPECIFIED))
+			{
+				XR_ENSURE(xrStopHapticFeedback(Session, &HapticActionInfo));
+			}
+			else
+			{
+				FOpenXRExtensionChainStructPtrs ScopedExtensionChainStructs;
+				if (Values.HapticBuffer != nullptr)
+				{
+					OpenXRHMD->GetApplyHapticFeedbackAddChainStructsDelegate().Broadcast(&HapticValue, ScopedExtensionChainStructs, Values.HapticBuffer);
+				}
+				XR_ENSURE(xrApplyHapticFeedback(Session, &HapticActionInfo, (const XrHapticBaseHeader*)&HapticValue));
+			}
+		}
+		if (Hand == (int32)EControllerHand::HMD)
+		{
+			XrHapticActionInfo HapticActionInfo;
+			HapticActionInfo.type = XR_TYPE_HAPTIC_ACTION_INFO;
+			HapticActionInfo.next = nullptr;
+			HapticActionInfo.subactionPath = XR_NULL_PATH;
+			HapticActionInfo.action = Controllers[EControllerHand::HMD].VibrationAction;
+			
+			if (Values.HapticBuffer == nullptr && (Values.Amplitude <= 0.0f || Values.Frequency < XR_FREQUENCY_UNSPECIFIED))
 			{
 				XR_ENSURE(xrStopHapticFeedback(Session, &HapticActionInfo));
 			}
diff --git a/Engine/Plugins/Runtime/OpenXR/Source/OpenXRInput/Private/OpenXRInput.h b/Engine/Plugins/Runtime/OpenXR/Source/OpenXRInput/Private/OpenXRInput.h
--- a/Engine/Plugins/Runtime/OpenXR/Source/OpenXRInput/Private/OpenXRInput.h
+++ b/Engine/Plugins/Runtime/OpenXR/Source/OpenXRInput/Private/OpenXRInput.h
@@ -12,6 +12,12 @@ class FOpenXRHMD;
 struct FInputActionKeyMapping;
 struct FInputAxisKeyMapping;
 
+// On some platforms the XrPath type becomes ambiguous for overloading
+FORCEINLINE uint32 GetTypeHash(const TPair<XrPath, XrPath>& Pair)
+{
+	return HashCombine(GetTypeHash((uint64)Pair.Key), GetTypeHash((uint64)Pair.Value));
+}
+
 class FOpenXRInputPlugin : public IOpenXRInputPlugin
 {
 public:
@@ -69,7 +75,7 @@ public:
 		// IMotionController overrides
 		virtual FName GetMotionControllerDeviceTypeName() const override;
 		virtual bool GetControllerOrientationAndPosition(const int32 ControllerIndex, const FName MotionSource, FRotator& OutOrientation, FVector& OutPosition, float WorldToMetersScale) const override;
-		virtual bool GetControllerOrientationAndPositionForTime(const int32 ControllerIndex, const FName MotionSource, FTimespan Time, bool& OutTimeWasUsed, FRotator& OutOrientation, FVector& OutPosition, bool& OutbProvidedLinearVelocity, FVector& OutLinearVelocity, bool& OutbProvidedAngularVelocity, FVector& OutAngularVelocityRadPerSec, float WorldToMetersScale) const override;
+		virtual bool GetControllerOrientationAndPositionForTime(const int32 ControllerIndex, const FName MotionSource, FTimespan Time, bool& OutTimeWasUsed, FRotator& OutOrientation, FVector& OutPosition, bool& OutbProvidedLinearVelocity, FVector& OutLinearVelocity, bool& OutbProvidedAngularVelocity, FVector& OutAngularVelocityRadPerSec, bool& OutbProvidedLinearAcceleration, FVector& OutLinearAcceleration, float WorldToMetersScale) const override;
 		virtual ETrackingStatus GetControllerTrackingStatus(const int32 ControllerIndex, const FName MotionSource) const override;
 		virtual bool GetControllerOrientationAndPosition(const int32 ControllerIndex, const EControllerHand DeviceHand, FRotator& OutOrientation, FVector& OutPosition, float WorldToMetersScale) const override { check(false); return false; }
 		virtual ETrackingStatus GetControllerTrackingStatus(const int32 ControllerIndex, const EControllerHand DeviceHand) const override { check(false); return ETrackingStatus::NotTracked; }
diff --git a/Engine/Source/Runtime/ApplicationCore/Public/GenericPlatform/IInputInterface.h b/Engine/Source/Runtime/ApplicationCore/Public/GenericPlatform/IInputInterface.h
--- a/Engine/Source/Runtime/ApplicationCore/Public/GenericPlatform/IInputInterface.h
+++ b/Engine/Source/Runtime/ApplicationCore/Public/GenericPlatform/IInputInterface.h
@@ -39,7 +39,7 @@ struct FForceFeedbackValues
 
 struct FHapticFeedbackBuffer
 {
-	TArray<uint8> RawData;
+	const uint8* RawData;
 	uint32 CurrentPtr;
 	int BufferLength;
 	int SamplesSent;
diff --git a/Engine/Source/Runtime/Engine/Classes/GameFramework/PlayerController.h b/Engine/Source/Runtime/Engine/Classes/GameFramework/PlayerController.h
--- a/Engine/Source/Runtime/Engine/Classes/GameFramework/PlayerController.h
+++ b/Engine/Source/Runtime/Engine/Classes/GameFramework/PlayerController.h
@@ -377,6 +377,7 @@ public:
 	TSharedPtr<struct FActiveHapticFeedbackEffect> ActiveHapticEffect_Left;
 	TSharedPtr<struct FActiveHapticFeedbackEffect> ActiveHapticEffect_Right;
 	TSharedPtr<struct FActiveHapticFeedbackEffect> ActiveHapticEffect_Gun;
+	TSharedPtr<struct FActiveHapticFeedbackEffect> ActiveHapticEffect_HMD;
 
 	/** Currently active force feedback weights */
 	FForceFeedbackValues ForceFeedbackValues;
diff --git a/Engine/Source/Runtime/Engine/Classes/Haptics/HapticFeedbackEffect_Base.h b/Engine/Source/Runtime/Engine/Classes/Haptics/HapticFeedbackEffect_Base.h
--- a/Engine/Source/Runtime/Engine/Classes/Haptics/HapticFeedbackEffect_Base.h
+++ b/Engine/Source/Runtime/Engine/Classes/Haptics/HapticFeedbackEffect_Base.h
@@ -5,6 +5,7 @@
 #include "CoreMinimal.h"
 #include "UObject/ObjectMacros.h"
 #include "UObject/Object.h"
+#include "GenericPlatform/IInputInterface.h"
 #include "HapticFeedbackEffect_Base.generated.h"
 
 struct FHapticFeedbackValues;
@@ -14,7 +15,7 @@ class UHapticFeedbackEffect_Base : public UObject
 {
 	GENERATED_UCLASS_BODY()
 
-	virtual void Initialize() {};
+	virtual void Initialize(FHapticFeedbackBuffer& HapticBuffer) {};
 
 	virtual void GetValues(const float EvalTime, FHapticFeedbackValues& Values);
 
@@ -43,13 +44,13 @@ struct ENGINE_API FActiveHapticFeedbackEffect
 		, PlayTime(0.f)
 	{
 		Scale = FMath::Clamp(InScale, 0.f, 10.f);
-		HapticEffect->Initialize();
+		HapticEffect->Initialize(HapticBuffer);
 	}
 
 	void Restart()
 	{
 		PlayTime = 0.f;
-		HapticEffect->Initialize();
+		HapticEffect->Initialize(HapticBuffer);
 	}
 
 	bool Update(const float DeltaTime, struct FHapticFeedbackValues& Values);
@@ -59,4 +60,5 @@ struct ENGINE_API FActiveHapticFeedbackEffect
 private:
 	float PlayTime;
 	float Scale;
+	FHapticFeedbackBuffer HapticBuffer;
 };
diff --git a/Engine/Source/Runtime/Engine/Classes/Haptics/HapticFeedbackEffect_Buffer.h b/Engine/Source/Runtime/Engine/Classes/Haptics/HapticFeedbackEffect_Buffer.h
--- a/Engine/Source/Runtime/Engine/Classes/Haptics/HapticFeedbackEffect_Buffer.h
+++ b/Engine/Source/Runtime/Engine/Classes/Haptics/HapticFeedbackEffect_Buffer.h
@@ -25,8 +25,5 @@ class UHapticFeedbackEffect_Buffer : public UHapticFeedbackEffect_Base
 
 	float GetDuration() const override;
 	
-	void Initialize() override;
-
-private:
-	FHapticFeedbackBuffer HapticBuffer;
+	void Initialize(FHapticFeedbackBuffer& HapticBuffer) override;
 };
diff --git a/Engine/Source/Runtime/Engine/Classes/Haptics/HapticFeedbackEffect_SoundWave.h b/Engine/Source/Runtime/Engine/Classes/Haptics/HapticFeedbackEffect_SoundWave.h
--- a/Engine/Source/Runtime/Engine/Classes/Haptics/HapticFeedbackEffect_SoundWave.h
+++ b/Engine/Source/Runtime/Engine/Classes/Haptics/HapticFeedbackEffect_SoundWave.h
@@ -20,7 +20,7 @@ class UHapticFeedbackEffect_SoundWave : public UHapticFeedbackEffect_Base
 
 	~UHapticFeedbackEffect_SoundWave();
 
-	void Initialize() override;
+	void Initialize(FHapticFeedbackBuffer& HapticBuffer) override;
 
 	void GetValues(const float EvalTime, FHapticFeedbackValues& Values) override;
 
@@ -30,5 +30,5 @@ private:
 	void PrepareSoundWaveBuffer();
 	bool bPrepared;
 
-	FHapticFeedbackBuffer HapticBuffer;
+	TArray<uint8> RawData;
 };
diff --git a/Engine/Source/Runtime/Engine/Private/GameFramework/HapticFeedbackEffect.cpp b/Engine/Source/Runtime/Engine/Private/GameFramework/HapticFeedbackEffect.cpp
--- a/Engine/Source/Runtime/Engine/Private/GameFramework/HapticFeedbackEffect.cpp
+++ b/Engine/Source/Runtime/Engine/Private/GameFramework/HapticFeedbackEffect.cpp
@@ -25,13 +25,23 @@ bool FActiveHapticFeedbackEffect::Update(const float DeltaTime, FHapticFeedbackV
 		return false;
 	}
 
+	HapticBuffer.RawData = nullptr;
+	Values.HapticBuffer = &HapticBuffer;
 	HapticEffect->GetValues(PlayTime, Values);
+	// Don't return a HapticBuffer if the effect didn't fill in RawData.
+	// Previously this buffer was owned by the HapticEffect itself, but that prevents
+	// playing the same effect on multiple controllers simultaneously.
+	if (HapticBuffer.RawData == nullptr)
+	{
+		Values.HapticBuffer = nullptr;
+	}
 	Values.Amplitude *= Scale;
 	if (Values.HapticBuffer)
 	{
 		Values.HapticBuffer->ScaleFactor = Scale;
 		if (Values.HapticBuffer->bFinishedPlaying)
 		{
+			Values.HapticBuffer = nullptr;
 			return false;
 		}
 	}
@@ -92,25 +102,19 @@ float UHapticFeedbackEffect_Curve::GetDuration() const
 UHapticFeedbackEffect_Buffer::UHapticFeedbackEffect_Buffer(const FObjectInitializer& ObjectInitializer)
 	: Super(ObjectInitializer)
 {
-	HapticBuffer.RawData.AddUninitialized(Amplitudes.Num());
-	HapticBuffer.CurrentPtr = 0;
-	HapticBuffer.BufferLength = Amplitudes.Num();
-	HapticBuffer.SamplesSent = 0;
-	HapticBuffer.bFinishedPlaying = false;
-	HapticBuffer.SamplingRate = SampleRate;
 }
 
 UHapticFeedbackEffect_Buffer::~UHapticFeedbackEffect_Buffer()
 {
-	HapticBuffer.RawData.Empty();
 }
 
 
-void UHapticFeedbackEffect_Buffer::Initialize()
+void UHapticFeedbackEffect_Buffer::Initialize(FHapticFeedbackBuffer& HapticBuffer)
 {
 	HapticBuffer.CurrentPtr = 0;
 	HapticBuffer.SamplesSent = 0;
 	HapticBuffer.bFinishedPlaying = false;
+	HapticBuffer.RawData = nullptr;
 }
 
 void UHapticFeedbackEffect_Buffer::GetValues(const float EvalTime, FHapticFeedbackValues& Values)
@@ -119,7 +123,6 @@ void UHapticFeedbackEffect_Buffer::GetValues(const float EvalTime, FHapticFeedba
 
 	Values.Frequency = 1.0;
 	Values.Amplitude = ampidx < Amplitudes.Num() ? (float)Amplitudes[ampidx] / 255.f : 0.f;
-	Values.HapticBuffer = &HapticBuffer;
 }
 
 float UHapticFeedbackEffect_Buffer::GetDuration() const
@@ -139,26 +142,28 @@ UHapticFeedbackEffect_SoundWave::UHapticFeedbackEffect_SoundWave(const FObjectIn
 
 UHapticFeedbackEffect_SoundWave::~UHapticFeedbackEffect_SoundWave()
 {
-	HapticBuffer.RawData.Empty();
+	RawData.Empty();
 }
 
-void UHapticFeedbackEffect_SoundWave::Initialize()
+void UHapticFeedbackEffect_SoundWave::Initialize(FHapticFeedbackBuffer& HapticBuffer)
 {
 	if (!bPrepared)
 	{
 		PrepareSoundWaveBuffer();
 	}
+	HapticBuffer.BufferLength = RawData.Num();
 	HapticBuffer.CurrentPtr = 0;
 	HapticBuffer.SamplesSent = 0;
 	HapticBuffer.bFinishedPlaying = false;
+	HapticBuffer.SamplingRate = SoundWave->GetSampleRateForCurrentPlatform();
 }
 
 void UHapticFeedbackEffect_SoundWave::GetValues(const float EvalTime, FHapticFeedbackValues& Values)
 {
-	int ampidx = EvalTime * HapticBuffer.BufferLength/ SoundWave->GetDuration();
+	int ampidx = EvalTime * RawData.Num() / SoundWave->GetDuration();
 	Values.Frequency = 1.0;
-	Values.Amplitude = ampidx < HapticBuffer.BufferLength ? (float)HapticBuffer.RawData[ampidx] / 255.f : 0.f;
-	Values.HapticBuffer = &HapticBuffer;
+	Values.Amplitude = ampidx < RawData.Num() ? (float)RawData[ampidx] / 255.f : 0.f;
+	Values.HapticBuffer->RawData = RawData.GetData();
 }
 
 float UHapticFeedbackEffect_SoundWave::GetDuration() const
@@ -177,35 +182,26 @@ void UHapticFeedbackEffect_SoundWave::PrepareSoundWaveBuffer()
 	SoundWave->InitAudioResource(AD->GetRuntimeFormat(SoundWave));
 	uint8* PCMData = SoundWave->RawPCMData;
 	int32 RawPCMDataSize = SoundWave->RawPCMDataSize;
-	check((PCMData != nullptr) || (RawPCMDataSize == 0));
-	int32 SampleRate = SoundWave->GetSampleRateForCurrentPlatform();
-	int TargetFrequency = 320;
-	int TargetBufferSize = (RawPCMDataSize * TargetFrequency) / (SampleRate * 2) + 1; //2 because we're only using half of the 16bit source PCM buffer
-	HapticBuffer.BufferLength = TargetBufferSize;
-	HapticBuffer.RawData.AddUninitialized(TargetBufferSize);
-	HapticBuffer.CurrentPtr = 0;
-	HapticBuffer.SamplingRate = TargetFrequency;
+	check((PCMData != nullptr) || (RawPCMDataSize == 0));	
 
-	int previousTargetIndex = -1;
-	int currentMin = 0;
-	for (int i = 1; i < RawPCMDataSize; i += 2)
+	// Some platforms may need to resample the PCM data.  Such resampling should be performed at the platform specific plugin level
+	int32 NumChannels = SoundWave->NumChannels;
+	if (NumChannels > 1)
 	{
-		int targetIndex = i * TargetFrequency / (SampleRate * 2);
-		int val = PCMData[i];
-		if (val & 0x80)
-		{
-			val = ~val;
-		}
-		currentMin = FMath::Min(currentMin, val);
-
-		if (targetIndex != previousTargetIndex)
+		UE_LOG(LogTemp, Warning, TEXT("%s used for vibration has more than 1 channel. Only the first channel will be used."), *SoundWave->GetPathName());
+		check(RawPCMDataSize % (sizeof(int16) * NumChannels) == 0);
+		int32 NumSamples = RawPCMDataSize / sizeof(int16) / NumChannels;
+		RawData.AddUninitialized(NumSamples * sizeof(int16));
+		int16* SourceData = reinterpret_cast<int16*>(PCMData);
+		int16* DestData = reinterpret_cast<int16*>(RawData.GetData());
+		for (int32 i = 0; i < NumSamples; i++)
 		{
-
-			HapticBuffer.RawData[targetIndex] = val * 2;// *Scale;
-			previousTargetIndex = targetIndex;
-			currentMin = 0;
+			DestData[i] = SourceData[i * NumChannels];
 		}
 	}
+	else
+	{
+		RawData.Append(PCMData, RawPCMDataSize);
+	}
 	bPrepared = true;
-
 }
diff --git a/Engine/Source/Runtime/Engine/Private/PlayerController.cpp b/Engine/Source/Runtime/Engine/Private/PlayerController.cpp
--- a/Engine/Source/Runtime/Engine/Private/PlayerController.cpp
+++ b/Engine/Source/Runtime/Engine/Private/PlayerController.cpp
@@ -4214,6 +4214,11 @@ void APlayerController::PlayHapticEffect(UHapticFeedbackEffect_Base* HapticEffec
 		case EControllerHand::Gun:
 			ActiveHapticEffect_Gun.Reset();
 			ActiveHapticEffect_Gun = MakeShareable(new FActiveHapticFeedbackEffect(HapticEffect, Scale, bLoop));
+			break;
+		case EControllerHand::HMD:
+			ActiveHapticEffect_HMD.Reset();
+			ActiveHapticEffect_HMD = MakeShareable(new FActiveHapticFeedbackEffect(HapticEffect, Scale, bLoop));
+			break;
 		default:
 			UE_LOG(LogPlayerController, Warning, TEXT("Invalid hand specified (%d) for haptic feedback effect %s"), (int32)Hand, *HapticEffect->GetName());
 			break;
@@ -4260,6 +4265,10 @@ void APlayerController::SetHapticsByValue(const float Frequency, const float Amp
 	{
 		ActiveHapticEffect_Gun.Reset();
 	}
+	else if (Hand == EControllerHand::HMD)
+	{
+		ActiveHapticEffect_HMD.Reset();
+	}
 	else
 	{
 		UE_LOG(LogPlayerController, Warning, TEXT("Invalid hand specified (%d) for setting haptic feedback values (F: %f A: %f)"), (int32)Hand, Frequency, Amplitude);
@@ -4330,6 +4339,8 @@ void APlayerController::ProcessForceFeedbackAndHaptics(const float DeltaTime, co
 	bool bLeftHapticsNeedUpdate = false;
 	bool bRightHapticsNeedUpdate = false;
 	bool bGunHapticsNeedUpdate = false;
+	FHapticFeedbackValues HMDHaptics;
+	bool bHMDHapticsNeedUpdate = false;
 
 	// Always process feedback by default, but if the game is paused then only static
 	// effects that are flagged to play while paused will play
@@ -4428,6 +4439,16 @@ void APlayerController::ProcessForceFeedbackAndHaptics(const float DeltaTime, co
 
 				bGunHapticsNeedUpdate = true;
 			}
+			if (ActiveHapticEffect_HMD.IsValid())
+			{
+				const bool bPlaying = ActiveHapticEffect_HMD->Update(DeltaTime, HMDHaptics);
+				if (!bPlaying)
+				{
+					ActiveHapticEffect_HMD->bLoop ? ActiveHapticEffect_HMD->Restart() : ActiveHapticEffect_HMD.Reset();
+				}
+
+				bHMDHapticsNeedUpdate = true;
+			}
 		}
 	}
 
@@ -4466,6 +4487,10 @@ void APlayerController::ProcessForceFeedbackAndHaptics(const float DeltaTime, co
 					{
 						InputInterface->SetHapticFeedbackValues(ControllerId, (int32)EControllerHand::Gun, GunHaptics);
 					}
+					if (bHMDHapticsNeedUpdate)
+					{
+						InputInterface->SetHapticFeedbackValues(ControllerId, (int32)EControllerHand::HMD, HMDHaptics);
+					}
 				}
 			}
 		}
diff --git a/Engine/Source/Runtime/HeadMountedDisplay/Private/HeadMountedDisplayFunctionLibrary.cpp b/Engine/Source/Runtime/HeadMountedDisplay/Private/HeadMountedDisplayFunctionLibrary.cpp
--- a/Engine/Source/Runtime/HeadMountedDisplay/Private/HeadMountedDisplayFunctionLibrary.cpp
+++ b/Engine/Source/Runtime/HeadMountedDisplay/Private/HeadMountedDisplayFunctionLibrary.cpp
@@ -530,7 +530,7 @@ void UHeadMountedDisplayFunctionLibrary::ClearXRTimedInputActionDelegate(const F
 	OnXRTimedInputActionDelegateMap.Remove(ActionName);
 }
 
-bool UHeadMountedDisplayFunctionLibrary::GetControllerTransformForTime(UObject* WorldContext, const int32 ControllerIndex, const FName MotionSource, FTimespan Time, bool& bTimeWasUsed, FRotator& Orientation, FVector& Position, bool& bProvidedLinearVelocity, FVector& LinearVelocity, bool& bProvidedAngularVelocity, FVector& AngularVelocityRadPerSec)
+bool UHeadMountedDisplayFunctionLibrary::GetControllerTransformForTime(UObject* WorldContext, const int32 ControllerIndex, const FName MotionSource, FTimespan Time, bool& bTimeWasUsed, FRotator& Orientation, FVector& Position, bool& bProvidedLinearVelocity, FVector& LinearVelocity, bool& bProvidedAngularVelocity, FVector& AngularVelocityRadPerSec, bool& bProvidedLinearAcceleration, FVector& LinearAcceleration)
 {
 	TArray<IMotionController*> MotionControllers = IModularFeatures::Get().GetModularFeatureImplementations<IMotionController>(IMotionController::GetModularFeatureName());
 	for (auto MotionController : MotionControllers)
@@ -542,7 +542,7 @@ bool UHeadMountedDisplayFunctionLibrary::GetControllerTransformForTime(UObject*
 
 		const float WorldToMetersScale = WorldContext ? WorldContext->GetWorld()->GetWorldSettings()->WorldToMeters : 100.0f;
 
-		const bool bGotTransform = MotionController->GetControllerOrientationAndPositionForTime(ControllerIndex, MotionSource, Time, bTimeWasUsed, Orientation, Position, bProvidedLinearVelocity, LinearVelocity, bProvidedAngularVelocity, AngularVelocityRadPerSec, WorldToMetersScale);
+		const bool bGotTransform = MotionController->GetControllerOrientationAndPositionForTime(ControllerIndex, MotionSource, Time, bTimeWasUsed, Orientation, Position, bProvidedLinearVelocity, LinearVelocity, bProvidedAngularVelocity, AngularVelocityRadPerSec, bProvidedLinearAcceleration, LinearAcceleration, WorldToMetersScale);
 		
 		if (bGotTransform)
 		{
diff --git a/Engine/Source/Runtime/HeadMountedDisplay/Private/XRMotionControllerBase.cpp b/Engine/Source/Runtime/HeadMountedDisplay/Private/XRMotionControllerBase.cpp
--- a/Engine/Source/Runtime/HeadMountedDisplay/Private/XRMotionControllerBase.cpp
+++ b/Engine/Source/Runtime/HeadMountedDisplay/Private/XRMotionControllerBase.cpp
@@ -59,12 +59,13 @@ bool FXRMotionControllerBase::GetControllerOrientationAndPosition(const int32 Co
 	return bSucess;
 }
 
-bool FXRMotionControllerBase::GetControllerOrientationAndPositionForTime(const int32 ControllerIndex, const FName MotionSource, FTimespan Time, bool& OutTimeWasUsed, FRotator& OutOrientation, FVector& OutPosition, bool& OutbProvidedLinearVelocity, FVector& OutLinearVelocity, bool& OutbProvidedAngularVelocity, FVector& OutAngularVelocityRadPerSec, float WorldToMetersScale) const
+bool FXRMotionControllerBase::GetControllerOrientationAndPositionForTime(const int32 ControllerIndex, const FName MotionSource, FTimespan Time, bool& OutTimeWasUsed, FRotator& OutOrientation, FVector& OutPosition, bool& OutbProvidedLinearVelocity, FVector& OutLinearVelocity, bool& OutbProvidedAngularVelocity, FVector& OutAngularVelocityRadPerSec, bool& OutbProvidedLinearAcceleration, FVector& OutLinearAcceleration, float WorldToMetersScale) const
 {
 	// Default implementation simply ignores the Timespan, no additional accuracy is provided by this call, velocities are not provided.
 	OutTimeWasUsed = false;
 	OutbProvidedLinearVelocity = false;
 	OutbProvidedAngularVelocity = false;
+	OutbProvidedLinearAcceleration = false;
 	return GetControllerOrientationAndPosition(ControllerIndex, MotionSource, OutOrientation, OutPosition, WorldToMetersScale);
 }
 
diff --git a/Engine/Source/Runtime/HeadMountedDisplay/Public/HeadMountedDisplayFunctionLibrary.h b/Engine/Source/Runtime/HeadMountedDisplay/Public/HeadMountedDisplayFunctionLibrary.h
--- a/Engine/Source/Runtime/HeadMountedDisplay/Public/HeadMountedDisplayFunctionLibrary.h
+++ b/Engine/Source/Runtime/HeadMountedDisplay/Public/HeadMountedDisplayFunctionLibrary.h
@@ -369,9 +369,10 @@ class HEADMOUNTEDDISPLAY_API UHeadMountedDisplayFunctionLibrary : public UBluepr
 	* The OpenXR spec suggests that 50ms in the past should return an accurate result.  There is no guarantee for the future, but the underlying system is likely to have been
 	* designed to predict out to about 50ms as well.
 	* On some platforms this  will always just return a cached position and rotation, ignoring time.  bTimeWasUsed will be false in that case.
+	* AngularVelocityRadPerSec is a vector whose direction is the axis of rotation and whoes length is the speed of rotation in radians per second.
 	*/
 	UFUNCTION(BlueprintCallable, Category = "Input|XRTracking")
-	static bool GetControllerTransformForTime(UObject* WorldContext, const int32 ControllerIndex, const FName MotionSource, FTimespan Time, bool& bTimeWasUsed, FRotator& Orientation, FVector& Position, bool& bProvidedLinearVelocity, FVector& LinearVelocity, bool& bProvidedAngularVelocity, FVector& AngularVelocityRadPerSec);
+	static bool GetControllerTransformForTime(UObject* WorldContext, const int32 ControllerIndex, const FName MotionSource, FTimespan Time, bool& bTimeWasUsed, FRotator& Orientation, FVector& Position, bool& bProvidedLinearVelocity, FVector& LinearVelocity, bool& bProvidedAngularVelocity, FVector& AngularVelocityRadPerSec, bool& bProvidedLinearAcceleration, FVector& LinearAcceleration);
 
 	/**
 	 * Get the bounds of the area where the user can freely move while remaining tracked centered around the specified origin
diff --git a/Engine/Source/Runtime/HeadMountedDisplay/Public/IMotionController.h b/Engine/Source/Runtime/HeadMountedDisplay/Public/IMotionController.h
--- a/Engine/Source/Runtime/HeadMountedDisplay/Public/IMotionController.h
+++ b/Engine/Source/Runtime/HeadMountedDisplay/Public/IMotionController.h
@@ -73,25 +73,27 @@ public:
 	virtual bool GetControllerOrientationAndPosition(const int32 ControllerIndex, const FName MotionSource, FRotator& OutOrientation, FVector& OutPosition, float WorldToMetersScale) const = 0;
 
 	/**
-	 * Returns the calibration-space orientation of the requested controller's hand at the specified time for potentially improved temporal precision, particularly fetching the controller transform when a button was pressed
-	 * on a platform that provides sub-frame timing for button presses.  This is only intended to work with times very near the current frame.  In general it should be called immediatly after receiving the button press. 
-	 * On many platforms this functionality is not supported and this function will set OutTimeWasUsed to false and then call GetControllerOrientationAndPosition, ignoring Time.
-	 *
-	 * @param ControllerIndex				The Unreal controller (player) index of the controller set
-	 * @param MotionSource					Which source, within the motion controller to get the orientation and position for
-	 * @param Time							The time at which we would like to query the orientation.
-	 * @param OutTimeWasUsed				(out) If true the time was used somehow to give a more temporally relevant orientation.  If false the time was ignored and a cached value returned.
-	 * @param OutOrientation				(out) If tracked, the orientation (in calibrated-space) of the controller in the specified hand
-	 * @param OutPosition					(out) If tracked, the position (in calibrated-space) of the controller in the specified hand
-	 * @param OutbProvidedLinearVelocity	(out) True if linear velocity was provided.
-	 * @param OutLinearVelocity				(out) The Linear velocity of the controller.
-	 * @param OutbProvidedAngularVelocity	(out) True if angular velocity was provided.
-	 * @param OutAngularVelocityRadPerSec   (out) The angular velocity of the controller in Radians per Second.
-	 * @param WorldToMetersScale			The world scaling factor.
-	 *
-	 * @return								True if the device requested is valid and tracked, false otherwise
-	 */
-	virtual bool GetControllerOrientationAndPositionForTime(const int32 ControllerIndex, const FName MotionSource, FTimespan Time, bool& OutTimeWasUsed, FRotator& OutOrientation, FVector& OutPosition, bool& OutbProvidedLinearVelocity, FVector& OutLinearVelocity, bool& OutbProvidedAngularVelocity, FVector& OutAngularVelocityRadPerSec, float WorldToMetersScale) const = 0;
+	* Returns the calibration-space orientation of the requested controller's hand at the specified time for potentially improved temporal precision, particularly fetching the controller transform when a button was pressed
+	* on a platform that provides sub-frame timing for button presses.  This is only intended to work with times very near the current frame.  In general it should be called immediatly after receiving the button press. 
+	* On many platforms this functionality is not supported and this function will set OutTimeWasUsed to false and then call GetControllerOrientationAndPosition, ignoring Time.
+	*
+	* @param ControllerIndex                The Unreal controller (player) index of the controller set
+	* @param MotionSource                   Which source, within the motion controller to get the orientation and position for
+	* @param Time                           The time at which we would like to query the orientation.
+	* @param OutTimeWasUsed                 (out) If true the time was used somehow to give a more temporally relevant orientation.  If false the time was ignored and a cached value returned.
+	* @param OutOrientation                 (out) If tracked, the orientation (in calibrated-space) of the controller in the specified hand
+	* @param OutPosition                    (out) If tracked, the position (in calibrated-space) of the controller in the specified hand
+	* @param OutbProvidedLinearVelocity     (out) True if linear velocity was provided.
+	* @param OutLinearVelocity              (out) The Linear velocity of the controller.
+	* @param OutbProvidedAngularVelocity    (out) True if angular velocity was provided.
+	* @param OutAngularVelocityRadPerSec    (out) The angular velocity of the controller in Radians per Second.
+	* @param OutbProvidedLinearAcceleration (out) True if linear acceleration was provided.
+	* @param OutLinearAcceleration          (out) The Linear acceleration of the controller.
+	* @param WorldToMetersScale             The world scaling factor.
+	*
+	* @return								True if the device requested is valid and tracked, false otherwise
+	*/
+	virtual bool GetControllerOrientationAndPositionForTime(const int32 ControllerIndex, const FName MotionSource, FTimespan Time, bool& OutTimeWasUsed, FRotator& OutOrientation, FVector& OutPosition, bool& OutbProvidedLinearVelocity, FVector& OutLinearVelocity, bool& OutbProvidedAngularVelocity, FVector& OutAngularVelocityRadPerSec, bool& OutbProvidedLinearAcceleration, FVector& OutLinearAcceleration, float WorldToMetersScale) const = 0;
 
 	/**
 	 * Returns the tracking status (e.g. not tracked, intertial-only, fully tracked) of the specified controller
diff --git a/Engine/Source/Runtime/HeadMountedDisplay/Public/XRMotionControllerBase.h b/Engine/Source/Runtime/HeadMountedDisplay/Public/XRMotionControllerBase.h
--- a/Engine/Source/Runtime/HeadMountedDisplay/Public/XRMotionControllerBase.h
+++ b/Engine/Source/Runtime/HeadMountedDisplay/Public/XRMotionControllerBase.h
@@ -16,7 +16,7 @@ public:
 
 	// Begin IMotionController interface
 	virtual bool GetControllerOrientationAndPosition(const int32 ControllerIndex, const FName MotionSource, FRotator& OutOrientation, FVector& OutPosition, float WorldToMetersScale) const;
-	virtual bool GetControllerOrientationAndPositionForTime(const int32 ControllerIndex, const FName MotionSource, FTimespan Time, bool& OutTimeWasUsed, FRotator& OutOrientation, FVector& OutPosition, bool& OutbProvidedLinearVelocity, FVector& OutLinearVelocity, bool& OutbProvidedAngularVelocity, FVector& OutAngularVelocityRadPerSec, float WorldToMetersScale) const;
+	virtual bool GetControllerOrientationAndPositionForTime(const int32 ControllerIndex, const FName MotionSource, FTimespan Time, bool& OutTimeWasUsed, FRotator& OutOrientation, FVector& OutPosition, bool& OutbProvidedLinearVelocity, FVector& OutLinearVelocity, bool& OutbProvidedAngularVelocity, FVector& OutAngularVelocityRadPerSec, bool& OutbProvidedLinearAcceleration, FVector& OutLinearAcceleration, float WorldToMetersScale) const;
 	virtual ETrackingStatus GetControllerTrackingStatus(const int32 ControllerIndex, const FName MotionSource) const;
 	virtual void EnumerateSources(TArray<FMotionControllerSource>& SourcesOut) const;
 	virtual float GetCustomParameterValue(const FName MotionSource, FName ParameterName, bool& bValueFound) const { bValueFound = false;  return 0.f; }
diff --git a/Engine/Source/Runtime/InputCore/Classes/InputCoreTypes.h b/Engine/Source/Runtime/InputCore/Classes/InputCoreTypes.h
--- a/Engine/Source/Runtime/InputCore/Classes/InputCoreTypes.h
+++ b/Engine/Source/Runtime/InputCore/Classes/InputCoreTypes.h
@@ -22,6 +22,7 @@ enum class EControllerHand : uint8
 	Pad,
 	ExternalCamera,
 	Gun,
+	HMD,
 	Special_1,
 	Special_2,
 	Special_3,
