PR #8627: Improve extensibility of the engine's OpenXR plugin (subset of PR #8621)  (Contributed by [REDACTED_PERSON])

diff --git a/Engine/Plugins/Runtime/OpenXR/Source/OpenXRHMD/Public/IOpenXRExtensionPlugin.h b/Engine/Plugins/Runtime/OpenXR/Source/OpenXRHMD/Public/IOpenXRExtensionPlugin.h
--- a/Engine/Plugins/Runtime/OpenXR/Source/OpenXRHMD/Public/IOpenXRExtensionPlugin.h
+++ b/Engine/Plugins/Runtime/OpenXR/Source/OpenXRHMD/Public/IOpenXRExtensionPlugin.h
@@ -224,6 +224,15 @@ public:
 	{
 	}
 
+	/**
+	* Add any action sets provided by the plugin to be attached as active to the session
+	* This allows a plugin to manage a custom actionset that will be active in xrSyncActions
+	*/
+	virtual void AddActionSets(TArray<XrActiveActionSet>& OutActionSets)
+	{
+	}
+
+
 	/**
 	* Use this callback to handle events that the OpenXR plugin doesn't handle itself
 	*/
diff --git a/Engine/Plugins/Runtime/OpenXR/Source/OpenXRInput/Private/OpenXRInput.cpp b/Engine/Plugins/Runtime/OpenXR/Source/OpenXRInput/Private/OpenXRInput.cpp
--- a/Engine/Plugins/Runtime/OpenXR/Source/OpenXRInput/Private/OpenXRInput.cpp
+++ b/Engine/Plugins/Runtime/OpenXR/Source/OpenXRInput/Private/OpenXRInput.cpp
@@ -109,6 +109,7 @@ FOpenXRInputPlugin::FOpenXRAction::FOpenXRAction(XrActionSet InActionSet, XrActi
 	Info.countSubactionPaths = SubactionPaths.Num();
 	Info.subactionPaths = SubactionPaths.GetData();
 	FCStringAnsi::Strcpy(Info.localizedActionName, XR_MAX_LOCALIZED_ACTION_NAME_SIZE, ActionName);
+
 	XR_ENSURE(xrCreateAction(Set, &Info, &Handle));
 }
 
@@ -393,6 +394,7 @@ void FOpenXRInputPlugin::FOpenXRInput::BuildActions()
 			InteractionProfile.interactionProfile = Profile.Path;
 			InteractionProfile.countSuggestedBindings = Profile.Bindings.Num();
 			InteractionProfile.suggestedBindings = Profile.Bindings.GetData();
+
 			XR_ENSURE(xrSuggestInteractionProfileBindings(Instance, &InteractionProfile));
 		}
 	}
@@ -528,10 +530,18 @@ void FOpenXRInputPlugin::FOpenXRInput::Tick(float DeltaTime)
 	{
 		if (!bActionsBound)
 		{
+			// Bind engine action sets
 			TArray<XrActionSet> BindActionSets;
 			for (auto && BindActionSet : ActionSets)
 				BindActionSets.Add(BindActionSet.actionSet);
 
+			// Bind plugin action sets
+			PluginActionSets.Empty();
+			for (IOpenXRExtensionPlugin* Plugin : OpenXRHMD->GetExtensionPlugins())
+				Plugin->AddActionSets(PluginActionSets);
+			for (auto&& BindActionSet : PluginActionSets)
+				BindActionSets.Add(BindActionSet.actionSet);
+
 			XrSessionActionSetsAttachInfo SessionActionSetsAttachInfo;
 			SessionActionSetsAttachInfo.type = XR_TYPE_SESSION_ACTION_SETS_ATTACH_INFO;
 			SessionActionSetsAttachInfo.next = nullptr;
@@ -553,12 +563,18 @@ void FOpenXRInputPlugin::FOpenXRInput::Tick(float DeltaTime)
 		XrActionsSyncInfo SyncInfo;
 		SyncInfo.type = XR_TYPE_ACTIONS_SYNC_INFO;
 		SyncInfo.next = nullptr;
+
+		PluginActionSets.Empty();
 		for (IOpenXRExtensionPlugin* Plugin : OpenXRHMD->GetExtensionPlugins())
 		{
+			Plugin->AddActionSets(PluginActionSets);
 			SyncInfo.next = Plugin->OnSyncActions(Session, SyncInfo.next);
 		}
-		SyncInfo.countActiveActionSets = ActionSets.Num();
-		SyncInfo.activeActionSets = ActionSets.GetData();
+
+		PluginActionSets.Append(ActionSets);
+		SyncInfo.countActiveActionSets = PluginActionSets.Num();
+		SyncInfo.activeActionSets = PluginActionSets.GetData();
+
 		XR_ENSURE(xrSyncActions(Session, &SyncInfo));
 	
 		for (IOpenXRExtensionPlugin* Plugin : OpenXRHMD->GetExtensionPlugins())
diff --git a/Engine/Plugins/Runtime/OpenXR/Source/OpenXRInput/Private/OpenXRInput.h b/Engine/Plugins/Runtime/OpenXR/Source/OpenXRInput/Private/OpenXRInput.h
--- a/Engine/Plugins/Runtime/OpenXR/Source/OpenXRInput/Private/OpenXRInput.h
+++ b/Engine/Plugins/Runtime/OpenXR/Source/OpenXRInput/Private/OpenXRInput.h
@@ -94,6 +94,7 @@ public:
 		FOpenXRHMD* OpenXRHMD;
 
 		TArray<XrActiveActionSet> ActionSets;
+		TArray<XrActiveActionSet> PluginActionSets;
 		TArray<XrPath> SubactionPaths;
 		TArray<FOpenXRAction> Actions;
 		TMap<EControllerHand, FOpenXRController> Controllers;
