OpenXRHMD: delete duplicated OnAcquireSwapchainImage
-Reverting merge mistake where a block of code was duplicated.

diff --git a/Engine/Plugins/Runtime/OpenXR/Source/OpenXRHMD/Private/OpenXRHMD.cpp b/Engine/Plugins/Runtime/OpenXR/Source/OpenXRHMD/Private/OpenXRHMD.cpp
--- a/Engine/Plugins/Runtime/OpenXR/Source/OpenXRHMD/Private/OpenXRHMD.cpp
+++ b/Engine/Plugins/Runtime/OpenXR/Source/OpenXRHMD/Private/OpenXRHMD.cpp
@@ -1819,13 +1819,6 @@ void FOpenXRHMD::OnBeginRendering_RenderThread(FRHICommandListImmediate& RHICmdL
 
 		SCOPED_NAMED_EVENT(EnqueueFrame, FColor::Red);
 
-		{
-			FReadScopeLock Lock(SessionHandleMutex);
-			for (IOpenXRExtensionPlugin* Module : ExtensionPlugins)
-			{
-				Module->OnAcquireSwapchainImage(Session);
-			}
-		}
 		// Reset the update flag on all layers
 		ForEachLayer([&](uint32 /* unused */, FOpenXRLayer& Layer)
 		{
