OpenXR optional VR support, OnAquireSwapchainImage moved to RHI thread.
-Exiting stereo does not necessarily cause the engine to quit out entirely.  If the exit is driven by the UE4 app, via xrRequestExitSession we do not quit out of the app when the session is destroyed.  There is also a cvar xr.OpenXRExitAppOnRuntimeDrivenSessionExit which, if true (the default) causes the app to exit if the session is destroyed without xrRequestExitSession (meaning the runtime destroyed it - suggesting a related launcher/ui overlay that is supposed to control app lifecycle).  In the editor we always 'exit', meaning we close the pie window.
-Adding OnBeginRendering_RenderThread hook to IOpenXRExtensionPlugin
-OnAquireSwapchainImage move is setup as an ifdef right now, defaulting to no change.  I'm not 100% sure this change is ready to go live in UE5 now.

diff --git a/Engine/Plugins/Runtime/OpenXR/Source/OpenXRHMD/Private/OpenXRHMD.cpp b/Engine/Plugins/Runtime/OpenXR/Source/OpenXRHMD/Private/OpenXRHMD.cpp
--- a/Engine/Plugins/Runtime/OpenXR/Source/OpenXRHMD/Private/OpenXRHMD.cpp
+++ b/Engine/Plugins/Runtime/OpenXR/Source/OpenXRHMD/Private/OpenXRHMD.cpp
@@ -37,6 +37,13 @@
 #define OPENXR_PAUSED_IDLE_FPS 10
 static const int64 OPENXR_SWAPCHAIN_WAIT_TIMEOUT = 100000000ll;		// 100ms in nanoseconds.
 
+static TAutoConsoleVariable<int32> CVarOpenXRExitAppOnRuntimeDrivenSessionExit(
+	TEXT("xr.OpenXRExitAppOnRuntimeDrivenSessionExit"),
+	1,
+	TEXT("If true, RequestExitApp will be called after we destroy the session because the state transitioned to XR_SESSION_STATE_EXITING or XR_SESSION_STATE_LOSS_PENDING and this is NOT the result of a call from the App to xrRequestExitSession.\n")
+	TEXT("The aniticipated situation is that the runtime is associated with a launcher application or has a runtime UI overlay which can tell openxr to exit vr and that in that context the app should also exit.  But maybe there are cases where it should not?  Set this CVAR to make it not.\n"),
+	ECVF_Default);
+
 namespace {
 	static TSet<XrEnvironmentBlendMode> SupportedBlendModes{ XR_ENVIRONMENT_BLEND_MODE_ALPHA_BLEND, XR_ENVIRONMENT_BLEND_MODE_ADDITIVE, XR_ENVIRONMENT_BLEND_MODE_OPAQUE };
 	static TSet<XrViewConfigurationType> SupportedViewConfigurations{ XR_VIEW_CONFIGURATION_TYPE_PRIMARY_STEREO, XR_VIEW_CONFIGURATION_TYPE_PRIMARY_QUAD_VARJO };
@@ -542,6 +549,10 @@ bool FOpenXRHMD::EnableStereo(bool stereo)
 		GEngine->bForceDisableFrameRateSmoothing = true;
 		if (OnStereoStartup())
 		{
+			if (!GIsEditor)
+			{
+				GEngine->SetMaxFPS(0);
+			}
 			StartSession();
 
 			FApp::SetUseVRFocus(true);
@@ -839,6 +850,7 @@ FOpenXRHMD::FOpenXRHMD(const FAutoRegister& AutoRegister, XrInstance InInstance,
 	, bIsReady(false)
 	, bIsRendering(false)
 	, bIsSynchronized(false)
+	, bIsExitingSessionByxrRequestExitSession(false)
 	, bNeedReAllocatedDepth(false)
 	, bNeedReBuildOcclusionMesh(true)
 	, bIsMobileMultiViewEnabled(false)
@@ -1278,6 +1290,9 @@ bool FOpenXRHMD::OnStereoStartup()
 	FWriteScopeLock Lock(SessionHandleMutex);
 	FWriteScopeLock DeviceLock(DeviceMutex);
 
+	check(Session == XR_NULL_HANDLE);
+	bIsExitingSessionByxrRequestExitSession = false;  // clear in case we requested exit for a previous session, but it ended in some other way before that happened.
+
 	XrSessionCreateInfo SessionInfo;
 	SessionInfo.type = XR_TYPE_SESSION_CREATE_INFO;
 	SessionInfo.next = RenderBridge->GetGraphicsBinding();
@@ -1387,6 +1402,8 @@ bool FOpenXRHMD::OnStereoTeardown()
 		FReadScopeLock Lock(SessionHandleMutex);
 		if (Session != XR_NULL_HANDLE)
 		{
+			UE_LOG(LogHMD, Verbose, TEXT("FOpenXRHMD::OnStereoTeardown() calling xrRequestExitSession"));
+			bIsExitingSessionByxrRequestExitSession = true;
 			Result = xrRequestExitSession(Session);
 		}
 	}
@@ -1679,6 +1696,11 @@ void FOpenXRHMD::OnBeginRendering_RenderThread(FRHICommandListImmediate& RHICmdL
 	ensure(IsInRenderingThread());
 	FReadScopeLock DeviceLock(DeviceMutex);
 
+	for (IOpenXRExtensionPlugin* Module : ExtensionPlugins)
+	{
+		Module->OnBeginRendering_RenderThread(Session);
+	}
+
 	const float WorldToMeters = GetWorldToMetersScale();
 	const FTransform InvTrackingToWorld = GetTrackingToWorldTransform().Inverse();
 
@@ -1825,6 +1847,8 @@ void FOpenXRHMD::OnBeginRendering_RenderThread(FRHICommandListImmediate& RHICmdL
 			Layer.bUpdateTexture = Layer.Desc.Flags & IStereoLayers::LAYER_FLAG_TEX_CONTINUOUS_UPDATE;
 		}, false);
 
+//TODO this is moving to OnBeginRendering_RHIThread
+#if 1
 		{
 			FReadScopeLock Lock(SessionHandleMutex);
 			for (IOpenXRExtensionPlugin* Module : ExtensionPlugins)
@@ -1832,6 +1856,7 @@ void FOpenXRHMD::OnBeginRendering_RenderThread(FRHICommandListImmediate& RHICmdL
 				Module->OnAcquireSwapchainImage(Session);
 			}
 		}
+#endif
 
 		FXRSwapChainPtr ColorSwapchain = PipelinedLayerStateRendering.ColorSwapchain;
 		FXRSwapChainPtr DepthSwapchain = PipelinedLayerStateRendering.DepthSwapchain;
@@ -2014,43 +2039,46 @@ bool FOpenXRHMD::OnStartGameFrame(FWorldContext& WorldContext)
 				bIsReady = false;
 				StopSession();
 			}
-			else if (SessionState.state == XR_SESSION_STATE_EXITING)
+			else if (SessionState.state == XR_SESSION_STATE_EXITING || SessionState.state == XR_SESSION_STATE_LOSS_PENDING)
 			{
 				// We need to make sure we unlock the frame rate again when exiting stereo while idle
 				if (!GIsEditor)
 				{
 					GEngine->SetMaxFPS(0);
 				}
-			}
+				
+				FApp::SetHasVRFocus(false);
 
-			FApp::SetHasVRFocus(SessionState.state == XR_SESSION_STATE_FOCUSED);
+				DestroySession();
 
-			if (SessionState.state != XR_SESSION_STATE_EXITING && SessionState.state != XR_SESSION_STATE_LOSS_PENDING)
-			{
-				break;
-			}
-		}
-		// Intentional fall-through
-		case XR_TYPE_EVENT_DATA_INSTANCE_LOSS_PENDING:
-		{
+				// Do we want to RequestExitApp the app after destoying the session?
+				// Yes if the app (ie ue4) did NOT requested the exit.
+				bool bExitApp = !bIsExitingSessionByxrRequestExitSession;
+				bIsExitingSessionByxrRequestExitSession = false;
+
+				// But only if this CVar is set to true.
+				bExitApp = bExitApp && (CVarOpenXRExitAppOnRuntimeDrivenSessionExit.GetValueOnAnyThread() != 0);
 #if WITH_EDITOR
-			if (GIsEditor)
-			{
-				FSceneViewport* SceneVP = FindSceneViewport();
-				if (SceneVP && SceneVP->IsStereoRenderingAllowed())
+				// But always if in the editor, because that doesn't actually exit.  See RequestExitApp().
+				bExitApp = bExitApp || GIsEditor;
+#endif
+				if (bExitApp)
 				{
-					TSharedPtr<SWindow> Window = SceneVP->FindWindow();
-					Window->RequestDestroyWindow();
+					RequestExitApp();
 				}
-			}
-			else
-#endif//WITH_EDITOR
-			{
-				// ApplicationWillTerminateDelegate will fire from inside of the RequestExit
-				FPlatformMisc::RequestExit(false);
+				break;
 			}
 
+			FApp::SetHasVRFocus(SessionState.state == XR_SESSION_STATE_FOCUSED);
+			
+			break;
+		}
+		case XR_TYPE_EVENT_DATA_INSTANCE_LOSS_PENDING:
+		{
 			DestroySession();
+			
+			// Instance loss is intended to support things like updating the active openxr runtime.  Currently we just require an app restart.
+			RequestExitApp();
 
 			break;
 		}
@@ -2084,6 +2112,28 @@ bool FOpenXRHMD::OnStartGameFrame(FWorldContext& WorldContext)
 	return true;
 }
 
+void FOpenXRHMD::RequestExitApp()
+{
+	UE_LOG(LogHMD, Log, TEXT("FOpenXRHMD is requesting app exit.  CurrentSessionState: %s"), OpenXRSessionStateToString(CurrentSessionState));
+
+#if WITH_EDITOR
+	if (GIsEditor)
+	{
+		FSceneViewport* SceneVP = FindSceneViewport();
+		if (SceneVP && SceneVP->IsStereoRenderingAllowed())
+		{
+			TSharedPtr<SWindow> Window = SceneVP->FindWindow();
+			Window->RequestDestroyWindow();
+		}
+	}
+	else
+#endif//WITH_EDITOR
+	{
+		// ApplicationWillTerminateDelegate will fire from inside of the RequestExit
+		FPlatformMisc::RequestExit(false);
+	}
+}
+
 void FOpenXRHMD::OnBeginRendering_RHIThread(const FPipelinedFrameState& InFrameState, FXRSwapChainPtr ColorSwapchain, FXRSwapChainPtr DepthSwapchain)
 {
 	ensure(IsInRenderingThread() || IsInRHIThread());
@@ -2123,6 +2173,14 @@ void FOpenXRHMD::OnBeginRendering_RHIThread(const FPipelinedFrameState& InFrameS
 		// We need a new swapchain image unless we've already acquired one for rendering
 		if (!bIsRendering && ColorSwapchain)
 		{
+//TODO this is moving from OnBeginRendering_RenderThread
+#if 0
+			for (IOpenXRExtensionPlugin* Module : ExtensionPlugins)
+			{
+				Module->OnAcquireSwapchainImage(Session);
+			}
+#endif
+
 			ColorSwapchain->IncrementSwapChainIndex_RHIThread();
 			ColorSwapchain->WaitCurrentImage_RHIThread(OPENXR_SWAPCHAIN_WAIT_TIMEOUT);
 			if (bDepthExtensionSupported && DepthSwapchain)
diff --git a/Engine/Plugins/Runtime/OpenXR/Source/OpenXRHMD/Private/OpenXRHMD.h b/Engine/Plugins/Runtime/OpenXR/Source/OpenXRHMD/Private/OpenXRHMD.h
--- a/Engine/Plugins/Runtime/OpenXR/Source/OpenXRHMD/Private/OpenXRHMD.h
+++ b/Engine/Plugins/Runtime/OpenXR/Source/OpenXRHMD/Private/OpenXRHMD.h
@@ -183,6 +183,8 @@ protected:
 	bool ReadNextEvent(XrEventDataBuffer* buffer);
 	void DestroySession();
 
+	void RequestExitApp();
+
 	void BuildOcclusionMeshes();
 	bool BuildOcclusionMesh(XrVisibilityMaskTypeKHR Type, int View, FHMDViewMesh& Mesh);
 
@@ -308,6 +310,7 @@ private:
 	TAtomic<bool>			bIsReady;
 	TAtomic<bool>			bIsRendering;
 	TAtomic<bool>			bIsSynchronized;
+	bool					bIsExitingSessionByxrRequestExitSession;
 	bool					bDepthExtensionSupported;
 	bool					bHiddenAreaMaskSupported;
 	bool					bViewConfigurationFovSupported;
diff --git a/Engine/Plugins/Runtime/OpenXR/Source/OpenXRHMD/Public/IOpenXRExtensionPlugin.h b/Engine/Plugins/Runtime/OpenXR/Source/OpenXRHMD/Public/IOpenXRExtensionPlugin.h
--- a/Engine/Plugins/Runtime/OpenXR/Source/OpenXRHMD/Public/IOpenXRExtensionPlugin.h
+++ b/Engine/Plugins/Runtime/OpenXR/Source/OpenXRHMD/Public/IOpenXRExtensionPlugin.h
@@ -317,7 +317,12 @@ public:
 		return InNext;
 	}
 
-	// OpenXRHMD::OnBeginRendering_RenderThread, before acquiring swapchain
+	// OpenXRHMD::OnBeginRendering_RenderThread
+	virtual void OnBeginRendering_RenderThread(XrSession InSession)
+	{
+	}
+
+	// OpenXRHMD::OnBeginRendering_RHIThread, before acquiring swapchain
 	virtual void OnAcquireSwapchainImage(XrSession InSession)
 	{
 	}
