OpenXR removing IOpenXRExtensionPlugin::OnAcquireSwapchainImage, it is unused.

diff --git a/Engine/Plugins/Runtime/OpenXR/Source/OpenXRHMD/Private/OpenXRHMD.cpp b/Engine/Plugins/Runtime/OpenXR/Source/OpenXRHMD/Private/OpenXRHMD.cpp
--- a/Engine/Plugins/Runtime/OpenXR/Source/OpenXRHMD/Private/OpenXRHMD.cpp
+++ b/Engine/Plugins/Runtime/OpenXR/Source/OpenXRHMD/Private/OpenXRHMD.cpp
@@ -1847,17 +1847,6 @@ void FOpenXRHMD::OnBeginRendering_RenderThread(FRHICommandListImmediate& RHICmdL
 			Layer.bUpdateTexture = Layer.Desc.Flags & IStereoLayers::LAYER_FLAG_TEX_CONTINUOUS_UPDATE;
 		}, false);
 
-//TODO this is moving to OnBeginRendering_RHIThread
-#if 1
-		{
-			FReadScopeLock Lock(SessionHandleMutex);
-			for (IOpenXRExtensionPlugin* Module : ExtensionPlugins)
-			{
-				Module->OnAcquireSwapchainImage(Session);
-			}
-		}
-#endif
-
 		FXRSwapChainPtr ColorSwapchain = PipelinedLayerStateRendering.ColorSwapchain;
 		FXRSwapChainPtr DepthSwapchain = PipelinedLayerStateRendering.DepthSwapchain;
 		RHICmdList.EnqueueLambda([this, FrameState = PipelinedFrameStateRendering, ColorSwapchain, DepthSwapchain](FRHICommandListImmediate& InRHICmdList)
@@ -2173,14 +2162,6 @@ void FOpenXRHMD::OnBeginRendering_RHIThread(const FPipelinedFrameState& InFrameS
 		// We need a new swapchain image unless we've already acquired one for rendering
 		if (!bIsRendering && ColorSwapchain)
 		{
-//TODO this is moving from OnBeginRendering_RenderThread
-#if 0
-			for (IOpenXRExtensionPlugin* Module : ExtensionPlugins)
-			{
-				Module->OnAcquireSwapchainImage(Session);
-			}
-#endif
-
 			ColorSwapchain->IncrementSwapChainIndex_RHIThread();
 			ColorSwapchain->WaitCurrentImage_RHIThread(OPENXR_SWAPCHAIN_WAIT_TIMEOUT);
 			if (bDepthExtensionSupported && DepthSwapchain)
diff --git a/Engine/Plugins/Runtime/OpenXR/Source/OpenXRHMD/Public/IOpenXRExtensionPlugin.h b/Engine/Plugins/Runtime/OpenXR/Source/OpenXRHMD/Public/IOpenXRExtensionPlugin.h
--- a/Engine/Plugins/Runtime/OpenXR/Source/OpenXRHMD/Public/IOpenXRExtensionPlugin.h
+++ b/Engine/Plugins/Runtime/OpenXR/Source/OpenXRHMD/Public/IOpenXRExtensionPlugin.h
@@ -322,11 +322,6 @@ public:
 	{
 	}
 
-	// OpenXRHMD::OnBeginRendering_RHIThread, before acquiring swapchain
-	virtual void OnAcquireSwapchainImage(XrSession InSession)
-	{
-	}
-
 	// OpenXRHMD::OnBeginRendering_RHIThread
 	virtual const void* OnBeginFrame(XrSession InSession, XrTime DisplayTime, const void* InNext)
 	{
