OpenXR support for IInputDevice::SetDeviceProperty calling into an OpenXRExtensionPlugin.  This is used for some haptics.

diff --git a/Engine/Plugins/Runtime/OpenXR/Source/OpenXRHMD/Public/IOpenXRExtensionPlugin.h b/Engine/Plugins/Runtime/OpenXR/Source/OpenXRHMD/Public/IOpenXRExtensionPlugin.h
--- a/Engine/Plugins/Runtime/OpenXR/Source/OpenXRHMD/Public/IOpenXRExtensionPlugin.h
+++ b/Engine/Plugins/Runtime/OpenXR/Source/OpenXRHMD/Public/IOpenXRExtensionPlugin.h
@@ -9,7 +9,7 @@
 #include "ARTraceResult.h"
 #include "DefaultSpectatorScreenController.h"
 #include "UObject/SoftObjectPath.h"
-
+#include "GenericPlatform/IInputInterface.h"
 #include <openxr/openxr.h>
 
 class IOpenXRCustomAnchorSupport
@@ -364,7 +364,11 @@ public:
 	virtual void PostSyncActions(XrSession InSession)
 	{
 	}
-
+	
+	virtual void OnSetDeviceProperty(XrSession InSession, int32 ControllerId, const FInputDeviceProperty* Property)
+	{
+	}
+	
 	/**
 	 * Start the AR system.
 	 *
diff --git a/Engine/Plugins/Runtime/OpenXR/Source/OpenXRInput/Private/OpenXRInput.cpp b/Engine/Plugins/Runtime/OpenXR/Source/OpenXRInput/Private/OpenXRInput.cpp
--- a/Engine/Plugins/Runtime/OpenXR/Source/OpenXRInput/Private/OpenXRInput.cpp
+++ b/Engine/Plugins/Runtime/OpenXR/Source/OpenXRInput/Private/OpenXRInput.cpp
@@ -741,6 +741,14 @@ void FOpenXRInputPlugin::FOpenXRInput::SetChannelValues(int32 ControllerId, cons
 		rightHaptics);
 }
 
+void FOpenXRInputPlugin::FOpenXRInput::SetDeviceProperty(int32 ControllerId, const FInputDeviceProperty* Property)
+{
+	for (IOpenXRExtensionPlugin* Module : OpenXRHMD->GetExtensionPlugins())
+	{
+		Module->OnSetDeviceProperty(OpenXRHMD->GetSession(), ControllerId, Property);
+	}
+}
+
 FName FOpenXRInputPlugin::FOpenXRInput::GetMotionControllerDeviceTypeName() const
 {
 	return FName(TEXT("OpenXR"));
diff --git a/Engine/Plugins/Runtime/OpenXR/Source/OpenXRInput/Private/OpenXRInput.h b/Engine/Plugins/Runtime/OpenXR/Source/OpenXRInput/Private/OpenXRInput.h
--- a/Engine/Plugins/Runtime/OpenXR/Source/OpenXRInput/Private/OpenXRInput.h
+++ b/Engine/Plugins/Runtime/OpenXR/Source/OpenXRInput/Private/OpenXRInput.h
@@ -71,7 +71,7 @@ public:
 		virtual bool Exec(UWorld* InWorld, const TCHAR* Cmd, FOutputDevice& Ar) override;
 		virtual void SetChannelValue(int32 ControllerId, FForceFeedbackChannelType ChannelType, float Value) override;
 		virtual void SetChannelValues(int32 ControllerId, const FForceFeedbackValues &values) override;
-
+		virtual void SetDeviceProperty(int32 ControllerId, const FInputDeviceProperty* Property) override;
 		// IMotionController overrides
 		virtual FName GetMotionControllerDeviceTypeName() const override;
 		virtual bool GetControllerOrientationAndPosition(const int32 ControllerIndex, const FName MotionSource, FRotator& OutOrientation, FVector& OutPosition, float WorldToMetersScale) const override;
