Check for null pointer

#rnx

diff --git a/Engine/Plugins/Runtime/OpenXR/Source/OpenXRInput/Private/OpenXRInput.cpp b/Engine/Plugins/Runtime/OpenXR/Source/OpenXRInput/Private/OpenXRInput.cpp
--- a/Engine/Plugins/Runtime/OpenXR/Source/OpenXRInput/Private/OpenXRInput.cpp
+++ b/Engine/Plugins/Runtime/OpenXR/Source/OpenXRInput/Private/OpenXRInput.cpp
@@ -743,6 +743,11 @@ void FOpenXRInputPlugin::FOpenXRInput::SetChannelValues(int32 ControllerId, cons
 
 void FOpenXRInputPlugin::FOpenXRInput::SetDeviceProperty(int32 ControllerId, const FInputDeviceProperty* Property)
 {
+	if (OpenXRHMD == nullptr)
+	{
+		return;
+	}
+
 	for (IOpenXRExtensionPlugin* Module : OpenXRHMD->GetExtensionPlugins())
 	{
 		Module->OnSetDeviceProperty(OpenXRHMD->GetSession(), ControllerId, Property);
