PR #8648: Support OpenXR Android loaders that do not require custom loader code (Contributed by [REDACTED_PERSON])

diff --git a/Engine/Plugins/Runtime/Oculus/OculusVR/Source/OculusOpenXRHMD/Private/OculusOpenXRHMD.cpp b/Engine/Plugins/Runtime/Oculus/OculusVR/Source/OculusOpenXRHMD/Private/OculusOpenXRHMD.cpp
--- a/Engine/Plugins/Runtime/Oculus/OculusVR/Source/OculusOpenXRHMD/Private/OculusOpenXRHMD.cpp
+++ b/Engine/Plugins/Runtime/Oculus/OculusVR/Source/OculusOpenXRHMD/Private/OculusOpenXRHMD.cpp
@@ -13,52 +13,6 @@
 
 DEFINE_LOG_CATEGORY(LogOculusOpenXRPlugin);
 
-bool FOculusOpenXRHMD::GetCustomLoader(PFN_xrGetInstanceProcAddr* OutGetProcAddr)
-{
-#if PLATFORM_ANDROID
-	// clear errors
-	dlerror();
-
-	LoaderHandle = FPlatformProcess::GetDllHandle(TEXT("libopenxr_loader.so"));
-	if (LoaderHandle == nullptr)
-	{
-		UE_LOG(LogOculusOpenXRPlugin, Error, TEXT("Unable to load libopenxr_loader.so, error %s"), ANSI_TO_TCHAR(dlerror()));
-		return false;
-	}
-
-	// clear errors
-	dlerror();
-
-	PFN_xrGetInstanceProcAddr xrGetInstanceProcAddrPtr = (PFN_xrGetInstanceProcAddr)FPlatformProcess::GetDllExport(LoaderHandle, TEXT("xrGetInstanceProcAddr"));
-	if (xrGetInstanceProcAddrPtr == nullptr)
-	{
-		UE_LOG(LogOculusOpenXRPlugin, Error, TEXT("Unable to load OpenXR xrGetInstanceProcAddr, error %s"), ANSI_TO_TCHAR(dlerror()));
-		return false;
-	}
-	*OutGetProcAddr = xrGetInstanceProcAddrPtr;
-
-	extern struct android_app* GNativeAndroidApp;
-	PFN_xrInitializeLoaderKHR xrInitializeLoaderKHR;
-	xrGetInstanceProcAddrPtr(XR_NULL_HANDLE, "xrInitializeLoaderKHR", (PFN_xrVoidFunction*)&xrInitializeLoaderKHR);
-	if (xrInitializeLoaderKHR == nullptr)
-	{
-		UE_LOG(LogOculusOpenXRPlugin, Error, TEXT("Unable to load OpenXR xrInitializeLoaderKHR"));
-		return false;
-	}
-	XrLoaderInitInfoAndroidKHR LoaderInitializeInfoAndroid;
-	LoaderInitializeInfoAndroid.type = XR_TYPE_LOADER_INIT_INFO_ANDROID_KHR;
-	LoaderInitializeInfoAndroid.next = NULL;
-	LoaderInitializeInfoAndroid.applicationVM = GNativeAndroidApp->activity->vm;
-	LoaderInitializeInfoAndroid.applicationContext = GNativeAndroidApp->activity->clazz;
-	XR_ENSURE(xrInitializeLoaderKHR((XrLoaderInitInfoBaseHeaderKHR*)&LoaderInitializeInfoAndroid));
-
-	UE_LOG(LogOculusOpenXRPlugin, Log, TEXT("Loaded Oculus OpenXR Loader"));
-	return true;
-#else //PLATFORM_ANDROID
-	return false;
-#endif //PLATFORM_ANDROID
-}
-
 bool FOculusOpenXRHMD::GetRequiredExtensions(TArray<const ANSICHAR*>& OutExtensions)
 {
 	return true;
@@ -147,6 +101,8 @@ const void* FOculusOpenXRHMD::OnEndProjectionLayer(XrSession InSession, int32 In
 	// XR_COMPOSITION_LAYER_UNPREMULTIPLIED_ALPHA_BIT is required right now because the Oculus mobile runtime blends using alpha otherwise,
 	// and we don't have proper inverse alpha support in OpenXR yet (once OpenXR supports inverse alpha, or we change the runtime behavior, remove this)
 	OutFlags |= XR_COMPOSITION_LAYER_CORRECT_CHROMATIC_ABERRATION_BIT;
+	OutFlags |= XR_COMPOSITION_LAYER_UNPREMULTIPLIED_ALPHA_BIT;
+
 	return InNext;
 }
 
diff --git a/Engine/Plugins/Runtime/Oculus/OculusVR/Source/OculusOpenXRHMD/Private/OculusOpenXRHMD.h b/Engine/Plugins/Runtime/Oculus/OculusVR/Source/OculusOpenXRHMD/Private/OculusOpenXRHMD.h
--- a/Engine/Plugins/Runtime/Oculus/OculusVR/Source/OculusOpenXRHMD/Private/OculusOpenXRHMD.h
+++ b/Engine/Plugins/Runtime/Oculus/OculusVR/Source/OculusOpenXRHMD/Private/OculusOpenXRHMD.h
@@ -33,7 +33,6 @@ public:
 		}
 	}
 
-	virtual bool GetCustomLoader(PFN_xrGetInstanceProcAddr* OutGetProcAddr) override;
 	virtual bool GetRequiredExtensions(TArray<const ANSICHAR*>& OutExtensions) override;
 	virtual bool GetInteractionProfile(XrInstance InInstance, FString& OutKeyPrefix, XrPath& OutPath, bool& OutHasHaptics) override;
 	virtual bool GetSpectatorScreenController(FHeadMountedDisplayBase* InHMDBase, TUniquePtr<FDefaultSpectatorScreenController>& OutSpectatorScreenController) override;
diff --git a/Engine/Plugins/Runtime/OpenXR/Source/OpenXRHMD/OpenXRHMD.Build.cs b/Engine/Plugins/Runtime/OpenXR/Source/OpenXRHMD/OpenXRHMD.Build.cs
--- a/Engine/Plugins/Runtime/OpenXR/Source/OpenXRHMD/OpenXRHMD.Build.cs
+++ b/Engine/Plugins/Runtime/OpenXR/Source/OpenXRHMD/OpenXRHMD.Build.cs
@@ -51,6 +51,7 @@ namespace UnrealBuildTool.Rules
                     "Slate",
                     "SlateCore",
 					"AugmentedReality",
+					"EngineSettings",
 				}
 				);
 
diff --git a/Engine/Plugins/Runtime/OpenXR/Source/OpenXRHMD/Private/OpenXRHMD.cpp b/Engine/Plugins/Runtime/OpenXR/Source/OpenXRHMD/Private/OpenXRHMD.cpp
--- a/Engine/Plugins/Runtime/OpenXR/Source/OpenXRHMD/Private/OpenXRHMD.cpp
+++ b/Engine/Plugins/Runtime/OpenXR/Source/OpenXRHMD/Private/OpenXRHMD.cpp
@@ -28,6 +28,7 @@
 #include "IHandTracker.h"
 #include "PixelShaderUtils.h"
 #include "ScenePrivate.h"
+#include "GeneralProjectSettings.h"
 #include "Epic_openxr.h"
 
 #if WITH_EDITOR
@@ -963,8 +964,10 @@ FOpenXRHMD::FOpenXRHMD(const FAutoRegister& AutoRegister, XrInstance InInstance,
 		}
 	}
 
-#if PLATFORM_HOLOLENS
-	bIsStandaloneStereoOnlyDevice = true;
+#if PLATFORM_HOLOLENS || PLATFORM_ANDROID
+	bool bStartInVR = false;
+	GConfig->GetBool(TEXT("/Script/EngineSettings.GeneralProjectSettings"), TEXT("bStartInVR"), bStartInVR, GGameIni);
+	bIsStandaloneStereoOnlyDevice = FParse::Param(FCommandLine::Get(), TEXT("vr")) || bStartInVR;
 #else
 	for (IOpenXRExtensionPlugin* Module : ExtensionPlugins)
 	{
diff --git a/Engine/Plugins/Runtime/OpenXR/Source/OpenXRHMD/Private/OpenXRHMDModule.cpp b/Engine/Plugins/Runtime/OpenXR/Source/OpenXRHMD/Private/OpenXRHMDModule.cpp
--- a/Engine/Plugins/Runtime/OpenXR/Source/OpenXRHMD/Private/OpenXRHMDModule.cpp
+++ b/Engine/Plugins/Runtime/OpenXR/Source/OpenXRHMD/Private/OpenXRHMDModule.cpp
@@ -7,6 +7,7 @@
 #include "IOpenXRExtensionPlugin.h"
 #include "IOpenXRARModule.h"
 #include "BuildSettings.h"
+#include "GeneralProjectSettings.h"
 
 #if PLATFORM_ANDROID
 #include <android_native_app_glue.h>
@@ -126,6 +127,12 @@ bool FOpenXRHMDModule::IsStandaloneStereoOnlyDevice()
 				return true;
 			}
 		}
+
+#if PLATFORM_HOLOLENS || PLATFORM_ANDROID
+		bool bStartInVR = false;
+		GConfig->GetBool(TEXT("/Script/EngineSettings.GeneralProjectSettings"), TEXT("bStartInVR"), bStartInVR, GGameIni); 
+		return FParse::Param(FCommandLine::Get(), TEXT("vr")) || bStartInVR;
+#endif
 	}
 	return false;
 }
@@ -313,7 +320,8 @@ PFN_xrGetInstanceProcAddr FOpenXRHMDModule::GetDefaultLoader()
 	FPlatformProcess::PopDllDirectory(*BinariesPath);
 #elif PLATFORM_LINUX
 	FString BinariesPath = FPaths::EngineDir() / FString(TEXT("Binaries/ThirdParty/OpenXR/linux/x86_64-unknown-linux-gnu"));
-	LoaderHandle = FPlatformProcess::GetDllHandle(*(BinariesPath / TEXT("libopenxr_loader.so")));
+	FString LoaderName = "libopenxr_loader.so";
+	LoaderHandle = FPlatformProcess::GetDllHandle(*(BinariesPath / LoaderName)); 
 #elif PLATFORM_HOLOLENS
 #ifndef PLATFORM_64BITS
 #error HoloLens platform does not currently support 32-bit. 32-bit OpenXR loader binaries are needed.
@@ -327,15 +335,38 @@ PFN_xrGetInstanceProcAddr FOpenXRHMDModule::GetDefaultLoader()
 #error Unsupported CPU family for the HoloLens platform.
 #endif
 
-	LoaderHandle = FPlatformProcess::GetDllHandle(*(BinariesPath / "openxr_loader.dll"));
+	FString LoaderName = "openxr_loader.dll";
+	LoaderHandle = FPlatformProcess::GetDllHandle(*(BinariesPath / LoaderName)); 
+#elif PLATFORM_ANDROID
+	FString LoaderName = "libopenxr_loader.so";
+	LoaderHandle = FPlatformProcess::GetDllHandle(*LoaderName);
 #endif
 
 	if (!LoaderHandle)
 	{
-		UE_LOG(LogHMD, Log, TEXT("Failed to load openxr_loader.dll"));
+		UE_LOG(LogHMD, Log, TEXT("Failed to load %s"), *LoaderName);
+		return nullptr;
+	}
+
+	PFN_xrGetInstanceProcAddr OutGetProcAddr = (PFN_xrGetInstanceProcAddr)FPlatformProcess::GetDllExport(LoaderHandle, TEXT("xrGetInstanceProcAddr"));
+
+#if PLATFORM_ANDROID
+	PFN_xrInitializeLoaderKHR xrInitializeLoaderKHR;
+	OutGetProcAddr(XR_NULL_HANDLE, "xrInitializeLoaderKHR", (PFN_xrVoidFunction*)&xrInitializeLoaderKHR);
+	if (xrInitializeLoaderKHR == nullptr)
+	{
+		UE_LOG(LogHMD, Error, TEXT("Unable to load OpenXR Android xrInitializeLoaderKHR"));
 		return nullptr;
 	}
-	return (PFN_xrGetInstanceProcAddr)FPlatformProcess::GetDllExport(LoaderHandle, TEXT("xrGetInstanceProcAddr"));
+
+	XrLoaderInitInfoAndroidKHR LoaderInitializeInfoAndroid;
+	LoaderInitializeInfoAndroid.type = XR_TYPE_LOADER_INIT_INFO_ANDROID_KHR;
+	LoaderInitializeInfoAndroid.next = NULL;
+	LoaderInitializeInfoAndroid.applicationVM = GNativeAndroidApp->activity->vm;
+	LoaderInitializeInfoAndroid.applicationContext = GNativeAndroidApp->activity->clazz;
+	XR_ENSURE(xrInitializeLoaderKHR((XrLoaderInitInfoBaseHeaderKHR*)&LoaderInitializeInfoAndroid));
+#endif
+	return OutGetProcAddr;
 }
 
 bool FOpenXRHMDModule::EnableExtensions(const TArray<const ANSICHAR*>& RequiredExtensions, const TArray<const ANSICHAR*>& OptionalExtensions, TArray<const ANSICHAR*>& OutExtensions)
