Fix compile error in OpenXRModule when platform is unsupported.

diff --git a/Engine/Plugins/Runtime/OpenXR/Source/OpenXRHMD/Private/OpenXRHMDModule.cpp b/Engine/Plugins/Runtime/OpenXR/Source/OpenXRHMD/Private/OpenXRHMDModule.cpp
--- a/Engine/Plugins/Runtime/OpenXR/Source/OpenXRHMD/Private/OpenXRHMDModule.cpp
+++ b/Engine/Plugins/Runtime/OpenXR/Source/OpenXRHMD/Private/OpenXRHMDModule.cpp
@@ -344,7 +344,7 @@ PFN_xrGetInstanceProcAddr FOpenXRHMDModule::GetDefaultLoader()
 
 	if (!LoaderHandle)
 	{
-		UE_LOG(LogHMD, Log, TEXT("Failed to load %s"), *LoaderName);
+		UE_LOG(LogHMD, Log, TEXT("Failed to find OpenXR runtime loader."));
 		return nullptr;
 	}
 
