PR #8763: Refine OpenGL ES support to OpenXR on Android. (Contributed by [REDACTED_PERSON])

diff --git a/Engine/Plugins/Runtime/OpenXR/Source/OpenXRHMD/Private/OpenXRHMDModule.cpp b/Engine/Plugins/Runtime/OpenXR/Source/OpenXRHMD/Private/OpenXRHMDModule.cpp
--- a/Engine/Plugins/Runtime/OpenXR/Source/OpenXRHMD/Private/OpenXRHMDModule.cpp
+++ b/Engine/Plugins/Runtime/OpenXR/Source/OpenXRHMD/Private/OpenXRHMDModule.cpp
@@ -272,20 +272,19 @@ bool FOpenXRHMDModule::InitRenderBridge()
 	}
 	else
 #endif
-#ifdef XR_USE_GRAPHICS_API_OPENGL
-	#if PLATFORM_ANDROID
+#if defined(XR_USE_GRAPHICS_API_OPENGL_ES) && defined(XR_USE_PLATFORM_ANDROID)
 	if (RHIString == TEXT("OpenGL") && IsExtensionEnabled(XR_KHR_OPENGL_ES_ENABLE_EXTENSION_NAME))
 	{
 		RenderBridge = CreateRenderBridge_OpenGLES(Instance, System);
 	}
 	else
-	#else
+#endif
+#ifdef XR_USE_GRAPHICS_API_OPENGL
 	if (RHIString == TEXT("OpenGL") && IsExtensionEnabled(XR_KHR_OPENGL_ENABLE_EXTENSION_NAME))
 	{
 		RenderBridge = CreateRenderBridge_OpenGL(Instance, System);
 	}
 	else
-	#endif
 #endif
 #ifdef XR_USE_GRAPHICS_API_VULKAN
 	if (RHIString == TEXT("Vulkan") && IsExtensionEnabled(XR_KHR_VULKAN_ENABLE_EXTENSION_NAME))
diff --git a/Engine/Plugins/Runtime/OpenXR/Source/OpenXRHMD/Private/OpenXRHMD_RenderBridge.h b/Engine/Plugins/Runtime/OpenXR/Source/OpenXRHMD/Private/OpenXRHMD_RenderBridge.h
--- a/Engine/Plugins/Runtime/OpenXR/Source/OpenXRHMD/Private/OpenXRHMD_RenderBridge.h
+++ b/Engine/Plugins/Runtime/OpenXR/Source/OpenXRHMD/Private/OpenXRHMD_RenderBridge.h
@@ -62,6 +62,9 @@ FOpenXRRenderBridge* CreateRenderBridge_D3D11(XrInstance InInstance, XrSystemId
 #ifdef XR_USE_GRAPHICS_API_D3D12
 FOpenXRRenderBridge* CreateRenderBridge_D3D12(XrInstance InInstance, XrSystemId InSystem);
 #endif
+#ifdef XR_USE_GRAPHICS_API_OPENGL_ES
+FOpenXRRenderBridge* CreateRenderBridge_OpenGLES(XrInstance InInstance, XrSystemId InSystem);
+#endif
 #ifdef XR_USE_GRAPHICS_API_OPENGL
 FOpenXRRenderBridge* CreateRenderBridge_OpenGL(XrInstance InInstance, XrSystemId InSystem);
 #endif
diff --git a/Engine/Source/Runtime/OpenGLDrv/Private/OpenGLDevice.cpp b/Engine/Source/Runtime/OpenGLDrv/Private/OpenGLDevice.cpp
--- a/Engine/Source/Runtime/OpenGLDrv/Private/OpenGLDevice.cpp
+++ b/Engine/Source/Runtime/OpenGLDrv/Private/OpenGLDevice.cpp
@@ -494,6 +494,7 @@ static inline void SetupTextureFormat( EPixelFormat Format, const FOpenGLTexture
 {
 	GOpenGLTextureFormats[Format] = GLFormat;
 	GPixelFormats[Format].Supported = (GLFormat.Format != GL_NONE && (GLFormat.InternalFormat[0] != GL_NONE || GLFormat.InternalFormat[1] != GL_NONE));
+	GPixelFormats[Format].PlatformFormat = GLFormat.InternalFormat[0];
 }
 
 
