OpenXR: proper ETextureCreateFlags management for swapchains

ETextureCreateFlags were previously split into two sets (resolve + target), but we utilized them as one set because we are only allocating one texture, not the pair.
Instead of maintaining this divergent logic, we can unify them for OpenXR usage.

There is still the issue of IStereoRenderTargetManager exposing both sets of texture flags. The concrete classes that utilize RHICreateTargetableShaderResource*
still require the separable flags.

diff --git a/Engine/Plugins/Runtime/OpenXR/Source/OpenXRHMD/Private/OpenXRHMD.cpp b/Engine/Plugins/Runtime/OpenXR/Source/OpenXRHMD/Private/OpenXRHMD.cpp
--- a/Engine/Plugins/Runtime/OpenXR/Source/OpenXRHMD/Private/OpenXRHMD.cpp
+++ b/Engine/Plugins/Runtime/OpenXR/Source/OpenXRHMD/Private/OpenXRHMD.cpp
@@ -1595,19 +1595,31 @@ bool FOpenXRHMD::AllocateRenderTargetTexture(uint32 Index, uint32 SizeX, uint32
 		return false;
 	}
 
+	// We're only creating a 1x target here, but we don't know whether it'll be the targeted texture
+	// or the resolve texture. Because of this, we unify the input flags.
+	ETextureCreateFlags UnifiedCreateFlags = Flags | TargetableTextureFlags;
+
 	// This is not a static swapchain
-	Flags |= TexCreate_Dynamic;
+	UnifiedCreateFlags |= TexCreate_Dynamic;
 
 	// We need to ensure we can sample from the texture in CopyTexture
-	TargetableTextureFlags |= TexCreate_ShaderResource;
+	UnifiedCreateFlags |= TexCreate_ShaderResource;
+
+	// We assume this could be used as a resolve target
+	UnifiedCreateFlags |= TexCreate_ResolveTargetable;
+
+	// Some render APIs require us to present in RT layouts/configs,
+	// so even if app won't use this texture as RT, we need the flag.
+	UnifiedCreateFlags |= TexCreate_RenderTargetable;
 
 	// On mobile without HDR all render targets need to be marked sRGB
 	bool MobileHWsRGB = IsMobileColorsRGB() && IsMobilePlatform(GMaxRHIShaderPlatform);
 	if (MobileHWsRGB)
 	{
-		TargetableTextureFlags |= TexCreate_SRGB;
+		UnifiedCreateFlags |= TexCreate_SRGB;
 	}
 
+
 	// Temporary workaround to swapchain formats - OpenXR doesn't support 10-bit sRGB swapchains, so prefer 8-bit sRGB instead.
 	if (Format == PF_A2B10G10R10)
 	{
@@ -1622,7 +1634,7 @@ bool FOpenXRHMD::AllocateRenderTargetTexture(uint32 Index, uint32 SizeX, uint32
 	{
 		ensureMsgf(NumSamples == 1, TEXT("OpenXR supports MSAA swapchains, but engine logic expects the swapchain target to be 1x."));
 
-		Swapchain = RenderBridge->CreateSwapchain(Session, Format, SizeX, SizeY, bIsMobileMultiViewEnabled ? 2 : 1, NumMips, NumSamples, Flags, TargetableTextureFlags, ClearColor);
+		Swapchain = RenderBridge->CreateSwapchain(Session, Format, SizeX, SizeY, bIsMobileMultiViewEnabled ? 2 : 1, NumMips, NumSamples, UnifiedCreateFlags, ClearColor);
 		if (!Swapchain)
 		{
 			return false;
@@ -1667,8 +1679,18 @@ bool FOpenXRHMD::AllocateDepthTexture(uint32 Index, uint32 SizeX, uint32 SizeY,
 		return false;
 	}
 
+	// We're only creating a 1x target here, but we don't know whether it'll be the targeted texture
+	// or the resolve texture. Because of this, we unify the input flags.
+	ETextureCreateFlags UnifiedCreateFlags = Flags | TargetableTextureFlags;
+
 	// This is not a static swapchain
-	Flags |= TexCreate_Dynamic;
+	UnifiedCreateFlags |= TexCreate_Dynamic;
+
+	// We assume this could be used as a resolve target
+	UnifiedCreateFlags |= TexCreate_DepthStencilResolveTarget;
+
+	// We can't use the depth swapchain w/o this flag, so client should always pass it in
+	ensure(EnumHasAllFlags(UnifiedCreateFlags, TexCreate_DepthStencilTargetable));
 
 	FXRSwapChainPtr& Swapchain = PipelinedLayerStateRendering.DepthSwapchain;
 	const FRHITexture2D* const DepthSwapchainTexture = Swapchain == nullptr ? nullptr : Swapchain->GetTexture2DArray() ? Swapchain->GetTexture2DArray() : Swapchain->GetTexture2D();
@@ -1676,7 +1698,7 @@ bool FOpenXRHMD::AllocateDepthTexture(uint32 Index, uint32 SizeX, uint32 SizeY,
 	{
 		ensureMsgf(NumSamples == 1, TEXT("OpenXR supports MSAA swapchains, but engine logic expects the swapchain target to be 1x."));
 
-		Swapchain = RenderBridge->CreateSwapchain(Session, PF_DepthStencil, SizeX, SizeY, bIsMobileMultiViewEnabled ? 2 : 1, FMath::Max(NumMips, 1u), NumSamples, Flags, TargetableTextureFlags, FClearValueBinding::DepthFar);
+		Swapchain = RenderBridge->CreateSwapchain(Session, PF_DepthStencil, SizeX, SizeY, bIsMobileMultiViewEnabled ? 2 : 1, FMath::Max(NumMips, 1u), NumSamples, UnifiedCreateFlags, FClearValueBinding::DepthFar);
 		if (!Swapchain)
 		{
 			return false;
@@ -1738,8 +1760,7 @@ void FOpenXRHMD::OnBeginRendering_RenderThread(FRHICommandListImmediate& RHICmdL
 			1,
 			Texture->GetNumMips(),
 			Texture->GetNumSamples(),
-			Texture->GetFlags() | Flags,
-			TexCreate_RenderTargetable,
+			Texture->GetFlags() | Flags | TexCreate_RenderTargetable,
 			Texture->GetClearBinding());
 	};
 	if (GetStereoLayersDirty())
diff --git a/Engine/Plugins/Runtime/OpenXR/Source/OpenXRHMD/Private/OpenXRHMD_RenderBridge.cpp b/Engine/Plugins/Runtime/OpenXR/Source/OpenXRHMD/Private/OpenXRHMD_RenderBridge.cpp
--- a/Engine/Plugins/Runtime/OpenXR/Source/OpenXRHMD/Private/OpenXRHMD_RenderBridge.cpp
+++ b/Engine/Plugins/Runtime/OpenXR/Source/OpenXRHMD/Private/OpenXRHMD_RenderBridge.cpp
@@ -49,9 +49,9 @@ public:
 		return &Binding;
 	}
 
-	virtual FXRSwapChainPtr CreateSwapchain(XrSession InSession, uint8 Format, uint32 SizeX, uint32 SizeY, uint32 ArraySize, uint32 NumMips, uint32 NumSamples, ETextureCreateFlags Flags, ETextureCreateFlags TargetableTextureFlags, const FClearValueBinding& ClearValueBinding) override final
+	virtual FXRSwapChainPtr CreateSwapchain(XrSession InSession, uint8 Format, uint32 SizeX, uint32 SizeY, uint32 ArraySize, uint32 NumMips, uint32 NumSamples, ETextureCreateFlags CreateFlags, const FClearValueBinding& ClearValueBinding) override final
 	{
-		return CreateSwapchain_D3D11(InSession, Format, SizeX, SizeY, ArraySize, NumMips, NumSamples, Flags, TargetableTextureFlags, ClearValueBinding);
+		return CreateSwapchain_D3D11(InSession, Format, SizeX, SizeY, ArraySize, NumMips, NumSamples, CreateFlags, ClearValueBinding);
 	}
 
 private:
@@ -93,9 +93,9 @@ public:
 		return &Binding;
 	}
 
-	virtual FXRSwapChainPtr CreateSwapchain(XrSession InSession, uint8 Format, uint32 SizeX, uint32 SizeY, uint32 ArraySize, uint32 NumMips, uint32 NumSamples, ETextureCreateFlags Flags, ETextureCreateFlags TargetableTextureFlags, const FClearValueBinding& ClearValueBinding) override final
+	virtual FXRSwapChainPtr CreateSwapchain(XrSession InSession, uint8 Format, uint32 SizeX, uint32 SizeY, uint32 ArraySize, uint32 NumMips, uint32 NumSamples, ETextureCreateFlags CreateFlags, const FClearValueBinding& ClearValueBinding) override final
 	{
-		return CreateSwapchain_D3D12(InSession, Format, SizeX, SizeY, ArraySize, NumMips, NumSamples, Flags, TargetableTextureFlags, ClearValueBinding);
+		return CreateSwapchain_D3D12(InSession, Format, SizeX, SizeY, ArraySize, NumMips, NumSamples, CreateFlags, ClearValueBinding);
 	}
 
 private:
@@ -146,9 +146,9 @@ public:
 		return nullptr;
 	}
 
-	virtual FXRSwapChainPtr CreateSwapchain(XrSession InSession, uint8 Format, uint32 SizeX, uint32 SizeY, uint32 ArraySize, uint32 NumMips, uint32 NumSamples, ETextureCreateFlags Flags, ETextureCreateFlags TargetableTextureFlags, const FClearValueBinding& ClearValueBinding) override final
+	virtual FXRSwapChainPtr CreateSwapchain(XrSession InSession, uint8 Format, uint32 SizeX, uint32 SizeY, uint32 ArraySize, uint32 NumMips, uint32 NumSamples, ETextureCreateFlags CreateFlags, const FClearValueBinding& ClearValueBinding) override final
 	{
-		return CreateSwapchain_OpenGL(InSession, Format, SizeX, SizeY, ArraySize, NumMips, NumSamples, Flags, TargetableTextureFlags, ClearValueBinding);
+		return CreateSwapchain_OpenGL(InSession, Format, SizeX, SizeY, ArraySize, NumMips, NumSamples, CreateFlags, ClearValueBinding);
 	}
 
 private:
@@ -204,9 +204,9 @@ public:
 		return nullptr;
 	}
 
-	virtual FXRSwapChainPtr CreateSwapchain(XrSession InSession, uint8 Format, uint32 SizeX, uint32 SizeY, uint32 ArraySize, uint32 NumMips, uint32 NumSamples, ETextureCreateFlags Flags, ETextureCreateFlags TargetableTextureFlags, const FClearValueBinding& ClearValueBinding) override final
+	virtual FXRSwapChainPtr CreateSwapchain(XrSession InSession, uint8 Format, uint32 SizeX, uint32 SizeY, uint32 ArraySize, uint32 NumMips, uint32 NumSamples, ETextureCreateFlags CreateFlags, const FClearValueBinding& ClearValueBinding) override final
 	{
-		return CreateSwapchain_OpenGLES(InSession, Format, SizeX, SizeY, ArraySize, NumMips, NumSamples, Flags, TargetableTextureFlags, ClearValueBinding);
+		return CreateSwapchain_OpenGLES(InSession, Format, SizeX, SizeY, ArraySize, NumMips, NumSamples, CreateFlags, ClearValueBinding);
 	}
 
 private:
@@ -292,9 +292,9 @@ public:
 		return AdapterLuid;
 	}
 
-	virtual FXRSwapChainPtr CreateSwapchain(XrSession InSession, uint8 Format, uint32 SizeX, uint32 SizeY, uint32 ArraySize, uint32 NumMips, uint32 NumSamples, ETextureCreateFlags Flags, ETextureCreateFlags TargetableTextureFlags, const FClearValueBinding& ClearValueBinding) override final
+	virtual FXRSwapChainPtr CreateSwapchain(XrSession InSession, uint8 Format, uint32 SizeX, uint32 SizeY, uint32 ArraySize, uint32 NumMips, uint32 NumSamples, ETextureCreateFlags CreateFlags, const FClearValueBinding& ClearValueBinding) override final
 	{
-		return CreateSwapchain_Vulkan(InSession, Format, SizeX, SizeY, ArraySize, NumMips, NumSamples, Flags, TargetableTextureFlags, ClearValueBinding);
+		return CreateSwapchain_Vulkan(InSession, Format, SizeX, SizeY, ArraySize, NumMips, NumSamples, CreateFlags, ClearValueBinding);
 	}
 
 private:
diff --git a/Engine/Plugins/Runtime/OpenXR/Source/OpenXRHMD/Private/OpenXRHMD_RenderBridge.h b/Engine/Plugins/Runtime/OpenXR/Source/OpenXRHMD/Private/OpenXRHMD_RenderBridge.h
--- a/Engine/Plugins/Runtime/OpenXR/Source/OpenXRHMD/Private/OpenXRHMD_RenderBridge.h
+++ b/Engine/Plugins/Runtime/OpenXR/Source/OpenXRHMD/Private/OpenXRHMD_RenderBridge.h
@@ -25,9 +25,9 @@ public:
 	virtual uint64 GetGraphicsAdapterLuid() { return AdapterLuid; }
 
 	virtual void* GetGraphicsBinding() = 0;
-	virtual FXRSwapChainPtr CreateSwapchain(XrSession InSession, uint8 Format, uint32 SizeX, uint32 SizeY, uint32 ArraySize, uint32 NumMips, uint32 NumSamples, ETextureCreateFlags Flags, ETextureCreateFlags TargetableTextureFlags, const FClearValueBinding& ClearValueBinding) = 0;
+	virtual FXRSwapChainPtr CreateSwapchain(XrSession InSession, uint8 Format, uint32 SizeX, uint32 SizeY, uint32 ArraySize, uint32 NumMips, uint32 NumSamples, ETextureCreateFlags CreateFlags, const FClearValueBinding& ClearValueBinding) = 0;
 
-	FXRSwapChainPtr CreateSwapchain(XrSession InSession, FRHITexture2D* Template, ETextureCreateFlags Flags, ETextureCreateFlags TargetableTextureFlags)
+	FXRSwapChainPtr CreateSwapchain(XrSession InSession, FRHITexture2D* Template, ETextureCreateFlags CreateFlags)
 	{
 		if (!Template)
 		{
@@ -41,8 +41,7 @@ public:
 			1,
 			Template->GetNumMips(),
 			Template->GetNumSamples(),
-			Template->GetFlags() | Flags,
-			TargetableTextureFlags,
+			Template->GetFlags() | CreateFlags,
 			Template->GetClearBinding());
 	}
 
diff --git a/Engine/Plugins/Runtime/OpenXR/Source/OpenXRHMD/Private/OpenXRHMD_Swapchain.cpp b/Engine/Plugins/Runtime/OpenXR/Source/OpenXRHMD/Private/OpenXRHMD_Swapchain.cpp
--- a/Engine/Plugins/Runtime/OpenXR/Source/OpenXRHMD/Private/OpenXRHMD_Swapchain.cpp
+++ b/Engine/Plugins/Runtime/OpenXR/Source/OpenXRHMD/Private/OpenXRHMD_Swapchain.cpp
@@ -141,23 +141,23 @@ uint8 FOpenXRSwapchain::GetNearestSupportedSwapchainFormat(XrSession InSession,
 	return FallbackFormat;
 }
 
-XrSwapchain FOpenXRSwapchain::CreateSwapchain(XrSession InSession, uint32 PlatformFormat, uint32 SizeX, uint32 SizeY, uint32 ArraySize, uint32 NumMips, uint32 NumSamples, ETextureCreateFlags Flags, ETextureCreateFlags TargetableTextureFlags)
+XrSwapchain FOpenXRSwapchain::CreateSwapchain(XrSession InSession, uint32 PlatformFormat, uint32 SizeX, uint32 SizeY, uint32 ArraySize, uint32 NumMips, uint32 NumSamples, ETextureCreateFlags CreateFlags)
 {
 	// Need a mutable format so we can reinterpret an sRGB format into a linear format
 	XrSwapchainUsageFlags Usage = XR_SWAPCHAIN_USAGE_MUTABLE_FORMAT_BIT;
-	if (EnumHasAnyFlags(TargetableTextureFlags, TexCreate_RenderTargetable))
+	if (EnumHasAnyFlags(CreateFlags, TexCreate_RenderTargetable))
 	{
 		Usage |= XR_SWAPCHAIN_USAGE_COLOR_ATTACHMENT_BIT;
 	}
-	if (EnumHasAnyFlags(TargetableTextureFlags, TexCreate_DepthStencilTargetable))
+	if (EnumHasAnyFlags(CreateFlags, TexCreate_DepthStencilTargetable))
 	{
 		Usage |= XR_SWAPCHAIN_USAGE_DEPTH_STENCIL_ATTACHMENT_BIT;
 	}
-	if (EnumHasAnyFlags(TargetableTextureFlags, TexCreate_ShaderResource))
+	if (EnumHasAnyFlags(CreateFlags, TexCreate_ShaderResource))
 	{
 		Usage |= XR_SWAPCHAIN_USAGE_SAMPLED_BIT;
 	}
-	if (EnumHasAnyFlags(TargetableTextureFlags, TexCreate_UAV))
+	if (EnumHasAnyFlags(CreateFlags, TexCreate_UAV))
 	{
 		Usage |= XR_SWAPCHAIN_USAGE_UNORDERED_ACCESS_BIT;
 	}
@@ -166,7 +166,7 @@ XrSwapchain FOpenXRSwapchain::CreateSwapchain(XrSession InSession, uint32 Platfo
 	XrSwapchainCreateInfo info;
 	info.type = XR_TYPE_SWAPCHAIN_CREATE_INFO;
 	info.next = nullptr;
-	info.createFlags = EnumHasAnyFlags(Flags, TexCreate_Dynamic) ? 0 : XR_SWAPCHAIN_CREATE_STATIC_IMAGE_BIT;
+	info.createFlags = EnumHasAnyFlags(CreateFlags, TexCreate_Dynamic) ? 0 : XR_SWAPCHAIN_CREATE_STATIC_IMAGE_BIT;
 	info.usageFlags = Usage;
 	info.format = PlatformFormat;
 	info.sampleCount = NumSamples;
@@ -198,7 +198,7 @@ TArray<T> EnumerateImages(XrSwapchain InSwapchain, XrStructureType InType)
 }
 
 #ifdef XR_USE_GRAPHICS_API_D3D11
-FXRSwapChainPtr CreateSwapchain_D3D11(XrSession InSession, uint8 Format, uint32 SizeX, uint32 SizeY, uint32 ArraySize, uint32 NumMips, uint32 NumSamples, ETextureCreateFlags Flags, ETextureCreateFlags TargetableTextureFlags, const FClearValueBinding& ClearValueBinding)
+FXRSwapChainPtr CreateSwapchain_D3D11(XrSession InSession, uint8 Format, uint32 SizeX, uint32 SizeY, uint32 ArraySize, uint32 NumMips, uint32 NumSamples, ETextureCreateFlags CreateFlags, const FClearValueBinding& ClearValueBinding)
 {
 	TFunction<uint32(uint8)> ToPlatformFormat = [](uint8 InFormat)
 	{
@@ -217,7 +217,7 @@ FXRSwapChainPtr CreateSwapchain_D3D11(XrSession InSession, uint8 Format, uint32
 		return nullptr;
 	}
 
-	XrSwapchain Swapchain = FOpenXRSwapchain::CreateSwapchain(InSession, ToPlatformFormat(Format), SizeX, SizeY, ArraySize, NumMips, NumSamples, Flags, TargetableTextureFlags);
+	XrSwapchain Swapchain = FOpenXRSwapchain::CreateSwapchain(InSession, ToPlatformFormat(Format), SizeX, SizeY, ArraySize, NumMips, NumSamples, CreateFlags);
 	if (!Swapchain)
 	{
 		return nullptr;
@@ -227,7 +227,7 @@ FXRSwapChainPtr CreateSwapchain_D3D11(XrSession InSession, uint8 Format, uint32
 	TArray<XrSwapchainImageD3D11KHR> Images = EnumerateImages<XrSwapchainImageD3D11KHR>(Swapchain, XR_TYPE_SWAPCHAIN_IMAGE_D3D11_KHR);
 	for (const auto& Image : Images)
 	{
-		TextureChain.Add(static_cast<FTextureRHIRef>(GD3D11RHI->RHICreateTexture2DArrayFromResource(GPixelFormats[Format].UnrealFormat, TargetableTextureFlags, ClearValueBinding, Image.texture)));
+		TextureChain.Add(static_cast<FTextureRHIRef>(GD3D11RHI->RHICreateTexture2DArrayFromResource(GPixelFormats[Format].UnrealFormat, CreateFlags, ClearValueBinding, Image.texture)));
 	}
 	FTextureRHIRef ChainTarget = static_cast<FTextureRHIRef>(GDynamicRHI->RHICreateAliasedTexture((FTextureRHIRef&)TextureChain[0]));
 
@@ -236,7 +236,7 @@ FXRSwapChainPtr CreateSwapchain_D3D11(XrSession InSession, uint8 Format, uint32
 #endif
 
 #ifdef XR_USE_GRAPHICS_API_D3D12
-FXRSwapChainPtr CreateSwapchain_D3D12(XrSession InSession, uint8 Format, uint32 SizeX, uint32 SizeY, uint32 ArraySize, uint32 NumMips, uint32 NumSamples, ETextureCreateFlags Flags, ETextureCreateFlags TargetableTextureFlags, const FClearValueBinding& ClearValueBinding)
+FXRSwapChainPtr CreateSwapchain_D3D12(XrSession InSession, uint8 Format, uint32 SizeX, uint32 SizeY, uint32 ArraySize, uint32 NumMips, uint32 NumSamples, ETextureCreateFlags CreateFlags, const FClearValueBinding& ClearValueBinding)
 {
 	TFunction<uint32(uint8)> ToPlatformFormat = [](uint8 InFormat)
 	{
@@ -255,7 +255,7 @@ FXRSwapChainPtr CreateSwapchain_D3D12(XrSession InSession, uint8 Format, uint32
 		return nullptr;
 	}
 
-	XrSwapchain Swapchain = FOpenXRSwapchain::CreateSwapchain(InSession, ToPlatformFormat(Format), SizeX, SizeY, ArraySize, NumMips, NumSamples, Flags, TargetableTextureFlags);
+	XrSwapchain Swapchain = FOpenXRSwapchain::CreateSwapchain(InSession, ToPlatformFormat(Format), SizeX, SizeY, ArraySize, NumMips, NumSamples, CreateFlags);
 	if (!Swapchain)
 	{
 		return nullptr;
@@ -266,7 +266,7 @@ FXRSwapChainPtr CreateSwapchain_D3D12(XrSession InSession, uint8 Format, uint32
 	TArray<XrSwapchainImageD3D12KHR> Images = EnumerateImages<XrSwapchainImageD3D12KHR>(Swapchain, XR_TYPE_SWAPCHAIN_IMAGE_D3D12_KHR);
 	for (const auto& Image : Images)
 	{
-		TextureChain.Add(static_cast<FTextureRHIRef>(DynamicRHI->RHICreateTexture2DArrayFromResource(GPixelFormats[Format].UnrealFormat, TargetableTextureFlags, ClearValueBinding, Image.texture)));
+		TextureChain.Add(static_cast<FTextureRHIRef>(DynamicRHI->RHICreateTexture2DArrayFromResource(GPixelFormats[Format].UnrealFormat, CreateFlags, ClearValueBinding, Image.texture)));
 	}
 	FTextureRHIRef ChainTarget = static_cast<FTextureRHIRef>(GDynamicRHI->RHICreateAliasedTexture((FTextureRHIRef&)TextureChain[0]));
 
@@ -275,7 +275,7 @@ FXRSwapChainPtr CreateSwapchain_D3D12(XrSession InSession, uint8 Format, uint32
 #endif
 
 #ifdef XR_USE_GRAPHICS_API_OPENGL
-FXRSwapChainPtr CreateSwapchain_OpenGL(XrSession InSession, uint8 Format, uint32 SizeX, uint32 SizeY, uint32 ArraySize, uint32 NumMips, uint32 NumSamples, ETextureCreateFlags Flags, ETextureCreateFlags TargetableTextureFlags, const FClearValueBinding& ClearValueBinding)
+FXRSwapChainPtr CreateSwapchain_OpenGL(XrSession InSession, uint8 Format, uint32 SizeX, uint32 SizeY, uint32 ArraySize, uint32 NumMips, uint32 NumSamples, ETextureCreateFlags CreateFlags, const FClearValueBinding& ClearValueBinding)
 {
 	Format = FOpenXRSwapchain::GetNearestSupportedSwapchainFormat(InSession, Format);
 	if (!Format)
@@ -283,7 +283,7 @@ FXRSwapChainPtr CreateSwapchain_OpenGL(XrSession InSession, uint8 Format, uint32
 		return nullptr;
 	}
 
-	XrSwapchain Swapchain = FOpenXRSwapchain::CreateSwapchain(InSession, GPixelFormats[Format].PlatformFormat, SizeX, SizeY, ArraySize, NumMips, NumSamples, Flags, TargetableTextureFlags);
+	XrSwapchain Swapchain = FOpenXRSwapchain::CreateSwapchain(InSession, GPixelFormats[Format].PlatformFormat, SizeX, SizeY, ArraySize, NumMips, NumSamples, CreateFlags);
 	if (!Swapchain)
 	{
 		return nullptr;
@@ -294,7 +294,7 @@ FXRSwapChainPtr CreateSwapchain_OpenGL(XrSession InSession, uint8 Format, uint32
 	TArray<XrSwapchainImageOpenGLKHR> Images = EnumerateImages<XrSwapchainImageOpenGLKHR>(Swapchain, XR_TYPE_SWAPCHAIN_IMAGE_OPENGL_KHR);
 	for (const auto& Image : Images)
 	{
-		TextureChain.Add(static_cast<FTextureRHIRef>(DynamicRHI->RHICreateTexture2DArrayFromResource(GPixelFormats[Format].UnrealFormat, SizeX, SizeY, ArraySize, NumMips, NumSamples, 1, ClearValueBinding, Image.image, TargetableTextureFlags)));
+		TextureChain.Add(static_cast<FTextureRHIRef>(DynamicRHI->RHICreateTexture2DArrayFromResource(GPixelFormats[Format].UnrealFormat, SizeX, SizeY, ArraySize, NumMips, NumSamples, 1, ClearValueBinding, Image.image, CreateFlags)));
 	}
 	FTextureRHIRef ChainTarget = static_cast<FTextureRHIRef>(GDynamicRHI->RHICreateAliasedTexture((FTextureRHIRef&)TextureChain[0]));
 
@@ -303,7 +303,7 @@ FXRSwapChainPtr CreateSwapchain_OpenGL(XrSession InSession, uint8 Format, uint32
 #endif
 
 #ifdef XR_USE_GRAPHICS_API_OPENGL_ES
-FXRSwapChainPtr CreateSwapchain_OpenGLES(XrSession InSession, uint8 Format, uint32 SizeX, uint32 SizeY, uint32 ArraySize, uint32 NumMips, uint32 NumSamples, ETextureCreateFlags Flags, ETextureCreateFlags TargetableTextureFlags, const FClearValueBinding& ClearValueBinding)
+FXRSwapChainPtr CreateSwapchain_OpenGLES(XrSession InSession, uint8 Format, uint32 SizeX, uint32 SizeY, uint32 ArraySize, uint32 NumMips, uint32 NumSamples, ETextureCreateFlags CreateFlags, const FClearValueBinding& ClearValueBinding)
 {
 	Format = FOpenXRSwapchain::GetNearestSupportedSwapchainFormat(InSession, Format);
 	if (!Format)
@@ -311,7 +311,7 @@ FXRSwapChainPtr CreateSwapchain_OpenGLES(XrSession InSession, uint8 Format, uint
 		return nullptr;
 	}
 
-	XrSwapchain Swapchain = FOpenXRSwapchain::CreateSwapchain(InSession, GPixelFormats[Format].PlatformFormat, SizeX, SizeY, ArraySize, NumMips, NumSamples, Flags, TargetableTextureFlags);
+	XrSwapchain Swapchain = FOpenXRSwapchain::CreateSwapchain(InSession, GPixelFormats[Format].PlatformFormat, SizeX, SizeY, ArraySize, NumMips, NumSamples, CreateFlags);
 	if (!Swapchain)
 	{
 		return nullptr;
@@ -322,7 +322,7 @@ FXRSwapChainPtr CreateSwapchain_OpenGLES(XrSession InSession, uint8 Format, uint
 	TArray<XrSwapchainImageOpenGLESKHR> Images = EnumerateImages<XrSwapchainImageOpenGLESKHR>(Swapchain, XR_TYPE_SWAPCHAIN_IMAGE_OPENGL_ES_KHR);
 	for (const auto& Image : Images)
 	{
-		TextureChain.Add(static_cast<FTextureRHIRef>(DynamicRHI->RHICreateTexture2DArrayFromResource(GPixelFormats[Format].UnrealFormat, SizeX, SizeY, ArraySize, NumMips, NumSamples, 1, ClearValueBinding, Image.image, TargetableTextureFlags)));
+		TextureChain.Add(static_cast<FTextureRHIRef>(DynamicRHI->RHICreateTexture2DArrayFromResource(GPixelFormats[Format].UnrealFormat, SizeX, SizeY, ArraySize, NumMips, NumSamples, 1, ClearValueBinding, Image.image, CreateFlags)));
 	}
 	FTextureRHIRef ChainTarget = static_cast<FTextureRHIRef>(GDynamicRHI->RHICreateAliasedTexture((FTextureRHIRef&)TextureChain[0]));
 
@@ -331,7 +331,7 @@ FXRSwapChainPtr CreateSwapchain_OpenGLES(XrSession InSession, uint8 Format, uint
 #endif
 
 #ifdef XR_USE_GRAPHICS_API_VULKAN
-FXRSwapChainPtr CreateSwapchain_Vulkan(XrSession InSession, uint8 Format, uint32 SizeX, uint32 SizeY, uint32 ArraySize, uint32 NumMips, uint32 NumSamples, ETextureCreateFlags Flags, ETextureCreateFlags TargetableTextureFlags, const FClearValueBinding& ClearValueBinding)
+FXRSwapChainPtr CreateSwapchain_Vulkan(XrSession InSession, uint8 Format, uint32 SizeX, uint32 SizeY, uint32 ArraySize, uint32 NumMips, uint32 NumSamples, ETextureCreateFlags CreateFlags, const FClearValueBinding& ClearValueBinding)
 {
 	TFunction<uint32(uint8)> ToPlatformFormat = [](uint8 InFormat)
 	{
@@ -344,7 +344,7 @@ FXRSwapChainPtr CreateSwapchain_Vulkan(XrSession InSession, uint8 Format, uint32
 		return nullptr;
 	}
 
-	XrSwapchain Swapchain = FOpenXRSwapchain::CreateSwapchain(InSession, ToPlatformFormat(Format), SizeX, SizeY, ArraySize, NumMips, NumSamples, Flags, TargetableTextureFlags);
+	XrSwapchain Swapchain = FOpenXRSwapchain::CreateSwapchain(InSession, ToPlatformFormat(Format), SizeX, SizeY, ArraySize, NumMips, NumSamples, CreateFlags);
 	if (!Swapchain)
 	{
 		return nullptr;
@@ -354,7 +354,7 @@ FXRSwapChainPtr CreateSwapchain_Vulkan(XrSession InSession, uint8 Format, uint32
 	TArray<XrSwapchainImageVulkanKHR> Images = EnumerateImages<XrSwapchainImageVulkanKHR>(Swapchain, XR_TYPE_SWAPCHAIN_IMAGE_VULKAN_KHR);
 	for (const auto& Image : Images)
 	{
-		TextureChain.Add(static_cast<FTextureRHIRef>(GVulkanRHI->RHICreateTexture2DArrayFromResource(GPixelFormats[Format].UnrealFormat, SizeX, SizeY, ArraySize, NumMips, NumSamples, Image.image, TargetableTextureFlags)));
+		TextureChain.Add(static_cast<FTextureRHIRef>(GVulkanRHI->RHICreateTexture2DArrayFromResource(GPixelFormats[Format].UnrealFormat, SizeX, SizeY, ArraySize, NumMips, NumSamples, Image.image, CreateFlags)));
 	}
 	FTextureRHIRef ChainTarget = static_cast<FTextureRHIRef>(GDynamicRHI->RHICreateAliasedTexture((FTextureRHIRef&)TextureChain[0]));
 
diff --git a/Engine/Plugins/Runtime/OpenXR/Source/OpenXRHMD/Private/OpenXRHMD_Swapchain.h b/Engine/Plugins/Runtime/OpenXR/Source/OpenXRHMD/Private/OpenXRHMD_Swapchain.h
--- a/Engine/Plugins/Runtime/OpenXR/Source/OpenXRHMD/Private/OpenXRHMD_Swapchain.h
+++ b/Engine/Plugins/Runtime/OpenXR/Source/OpenXRHMD/Private/OpenXRHMD_Swapchain.h
@@ -20,7 +20,7 @@ public:
 	virtual void ReleaseCurrentImage_RHIThread() override final;
 
 	XrSwapchain GetHandle() { return Handle; }
-	static XrSwapchain CreateSwapchain(XrSession InSession, uint32 PlatformFormat, uint32 SizeX, uint32 SizeY, uint32 ArraySize, uint32 NumMips, uint32 NumSamples, ETextureCreateFlags Flags, ETextureCreateFlags TargetableTextureFlags);
+	static XrSwapchain CreateSwapchain(XrSession InSession, uint32 PlatformFormat, uint32 SizeX, uint32 SizeY, uint32 ArraySize, uint32 NumMips, uint32 NumSamples, ETextureCreateFlags CreateFlags);
 	static uint8 GetNearestSupportedSwapchainFormat(XrSession InSession, uint8 RequestedFormat, TFunction<uint32(uint8)> ToPlatformFormat = nullptr);
 
 protected:
@@ -29,17 +29,17 @@ protected:
 };
 
 #ifdef XR_USE_GRAPHICS_API_D3D11
-FXRSwapChainPtr CreateSwapchain_D3D11(XrSession InSession, uint8 Format, uint32 SizeX, uint32 SizeY, uint32 ArraySize, uint32 NumMips, uint32 NumSamples, ETextureCreateFlags Flags, ETextureCreateFlags TargetableTextureFlags, const FClearValueBinding& ClearValueBinding);
+FXRSwapChainPtr CreateSwapchain_D3D11(XrSession InSession, uint8 Format, uint32 SizeX, uint32 SizeY, uint32 ArraySize, uint32 NumMips, uint32 NumSamples, ETextureCreateFlags CreateFlags, const FClearValueBinding& ClearValueBinding);
 #endif
 #ifdef XR_USE_GRAPHICS_API_D3D12
-FXRSwapChainPtr CreateSwapchain_D3D12(XrSession InSession, uint8 Format, uint32 SizeX, uint32 SizeY, uint32 ArraySize, uint32 NumMips, uint32 NumSamples, ETextureCreateFlags Flags, ETextureCreateFlags TargetableTextureFlags, const FClearValueBinding& ClearValueBinding);
+FXRSwapChainPtr CreateSwapchain_D3D12(XrSession InSession, uint8 Format, uint32 SizeX, uint32 SizeY, uint32 ArraySize, uint32 NumMips, uint32 NumSamples, ETextureCreateFlags CreateFlags, const FClearValueBinding& ClearValueBinding);
 #endif
 #ifdef XR_USE_GRAPHICS_API_OPENGL
-FXRSwapChainPtr CreateSwapchain_OpenGL(XrSession InSession, uint8 Format, uint32 SizeX, uint32 SizeY, uint32 ArraySize, uint32 NumMips, uint32 NumSamples, ETextureCreateFlags Flags, ETextureCreateFlags TargetableTextureFlags, const FClearValueBinding& ClearValueBinding);
+FXRSwapChainPtr CreateSwapchain_OpenGL(XrSession InSession, uint8 Format, uint32 SizeX, uint32 SizeY, uint32 ArraySize, uint32 NumMips, uint32 NumSamples, ETextureCreateFlags CreateFlags, const FClearValueBinding& ClearValueBinding);
 #endif
 #ifdef XR_USE_GRAPHICS_API_OPENGL_ES
-FXRSwapChainPtr CreateSwapchain_OpenGLES(XrSession InSession, uint8 Format, uint32 SizeX, uint32 SizeY, uint32 ArraySize, uint32 NumMips, uint32 NumSamples, ETextureCreateFlags Flags, ETextureCreateFlags TargetableTextureFlags, const FClearValueBinding& ClearValueBinding);
+FXRSwapChainPtr CreateSwapchain_OpenGLES(XrSession InSession, uint8 Format, uint32 SizeX, uint32 SizeY, uint32 ArraySize, uint32 NumMips, uint32 NumSamples, ETextureCreateFlags CreateFlags, const FClearValueBinding& ClearValueBinding);
 #endif
 #ifdef XR_USE_GRAPHICS_API_VULKAN
-FXRSwapChainPtr CreateSwapchain_Vulkan(XrSession InSession, uint8 Format, uint32 SizeX, uint32 SizeY, uint32 ArraySize, uint32 NumMips, uint32 NumSamples, ETextureCreateFlags Flags, ETextureCreateFlags TargetableTextureFlags, const FClearValueBinding& ClearValueBinding);
+FXRSwapChainPtr CreateSwapchain_Vulkan(XrSession InSession, uint8 Format, uint32 SizeX, uint32 SizeY, uint32 ArraySize, uint32 NumMips, uint32 NumSamples, ETextureCreateFlags CreateFlags, const FClearValueBinding& ClearValueBinding);
 #endif
\ No newline at end of file
