PR #8714: Use custom OpenXR XrReferenceSpaceType from OpenXRExtensionPlugins (Contributed by [REDACTED_PERSON])

diff --git a/Engine/Plugins/Runtime/OpenXR/Source/OpenXRHMD/Private/OpenXRHMD.cpp b/Engine/Plugins/Runtime/OpenXR/Source/OpenXRHMD/Private/OpenXRHMD.cpp
--- a/Engine/Plugins/Runtime/OpenXR/Source/OpenXRHMD/Private/OpenXRHMD.cpp
+++ b/Engine/Plugins/Runtime/OpenXR/Source/OpenXRHMD/Private/OpenXRHMD.cpp
@@ -887,6 +887,7 @@ FOpenXRHMD::FOpenXRHMD(const FAutoRegister& AutoRegister, XrInstance InInstance,
 	, Session(XR_NULL_HANDLE)
 	, LocalSpace(XR_NULL_HANDLE)
 	, StageSpace(XR_NULL_HANDLE)
+	, CustomSpace(XR_NULL_HANDLE)
 	, TrackingSpaceType(XR_REFERENCE_SPACE_TYPE_STAGE)
 	, SelectedViewConfigurationType(XR_VIEW_CONFIGURATION_TYPE_MAX_ENUM)
 	, SelectedEnvironmentBlendMode(XR_ENVIRONMENT_BLEND_MODE_MAX_ENUM)
@@ -1364,8 +1365,26 @@ bool FOpenXRHMD::OnStereoStartup()
 	SpaceInfo.referenceSpaceType = XR_REFERENCE_SPACE_TYPE_LOCAL;
 	XR_ENSURE(xrCreateReferenceSpace(Session, &SpaceInfo, &LocalSpace));
 
-	// Prefer a stage space over a local space
-	if (Spaces.Contains(XR_REFERENCE_SPACE_TYPE_STAGE))
+	bool UseCustomReferenceSpaceType = false;
+	XrReferenceSpaceType CustomReferenceSpaceType;
+	for (IOpenXRExtensionPlugin* Module : ExtensionPlugins)
+	{
+		if (Module->UseCustomReferenceSpaceType(CustomReferenceSpaceType))
+		{
+			UseCustomReferenceSpaceType = true;
+			break;
+		}
+	}
+
+	// If a custom reference space is desired, try to use that.  
+	// Otherwise prefer a stage space over a local space
+	if (UseCustomReferenceSpaceType && Spaces.Contains(CustomReferenceSpaceType))
+	{
+		TrackingSpaceType = CustomReferenceSpaceType;
+		SpaceInfo.referenceSpaceType = TrackingSpaceType;
+		XR_ENSURE(xrCreateReferenceSpace(Session, &SpaceInfo, &CustomSpace));
+	}
+	else if (Spaces.Contains(XR_REFERENCE_SPACE_TYPE_STAGE))
 	{
 		TrackingSpaceType = XR_REFERENCE_SPACE_TYPE_STAGE;
 		SpaceInfo.referenceSpaceType = TrackingSpaceType;
diff --git a/Engine/Plugins/Runtime/OpenXR/Source/OpenXRHMD/Private/OpenXRHMD.h b/Engine/Plugins/Runtime/OpenXR/Source/OpenXRHMD/Private/OpenXRHMD.h
--- a/Engine/Plugins/Runtime/OpenXR/Source/OpenXRHMD/Private/OpenXRHMD.h
+++ b/Engine/Plugins/Runtime/OpenXR/Source/OpenXRHMD/Private/OpenXRHMD.h
@@ -297,6 +297,11 @@ public:
 	OPENXRHMD_API XrSession GetSession() { return Session; }
 	OPENXRHMD_API XrSpace GetTrackingSpace()
 	{
+		if (CustomSpace != XR_NULL_HANDLE)
+		{
+			return CustomSpace;
+		}
+		
 		return (TrackingSpaceType == XR_REFERENCE_SPACE_TYPE_STAGE) ? StageSpace : LocalSpace;
 	}
 	OPENXRHMD_API XrTime GetDisplayTime() const;
@@ -332,6 +337,7 @@ private:
 	XrSession				Session;
 	XrSpace					LocalSpace;
 	XrSpace					StageSpace;
+	XrSpace					CustomSpace;
 	XrReferenceSpaceType	TrackingSpaceType;
 	XrViewConfigurationType SelectedViewConfigurationType;
 	XrEnvironmentBlendMode  SelectedEnvironmentBlendMode;
diff --git a/Engine/Plugins/Runtime/OpenXR/Source/OpenXRHMD/Public/IOpenXRExtensionPlugin.h b/Engine/Plugins/Runtime/OpenXR/Source/OpenXRHMD/Public/IOpenXRExtensionPlugin.h
--- a/Engine/Plugins/Runtime/OpenXR/Source/OpenXRHMD/Public/IOpenXRExtensionPlugin.h
+++ b/Engine/Plugins/Runtime/OpenXR/Source/OpenXRHMD/Public/IOpenXRExtensionPlugin.h
@@ -364,10 +364,17 @@ public:
 	virtual void PostSyncActions(XrSession InSession)
 	{
 	}
+
 	
 	virtual void OnSetDeviceProperty(XrSession InSession, int32 ControllerId, const FInputDeviceProperty* Property)
 	{
 	}
+
+	/** Update OpenXRHMD to use reference space types other than view, local, and stage. */
+	virtual bool UseCustomReferenceSpaceType(XrReferenceSpaceType& OutReferenceSpaceType)
+	{
+		return false;
+	}
 	
 	/**
 	 * Start the AR system.
