Don't mark sRGB textures as mutable

Also adds comments explaining why we need this flag.

diff --git a/Engine/Plugins/Runtime/OpenXR/Source/OpenXRHMD/Private/OpenXRHMD_Swapchain.cpp b/Engine/Plugins/Runtime/OpenXR/Source/OpenXRHMD/Private/OpenXRHMD_Swapchain.cpp
--- a/Engine/Plugins/Runtime/OpenXR/Source/OpenXRHMD/Private/OpenXRHMD_Swapchain.cpp
+++ b/Engine/Plugins/Runtime/OpenXR/Source/OpenXRHMD/Private/OpenXRHMD_Swapchain.cpp
@@ -143,8 +143,19 @@ uint8 FOpenXRSwapchain::GetNearestSupportedSwapchainFormat(XrSession InSession,
 
 XrSwapchain FOpenXRSwapchain::CreateSwapchain(XrSession InSession, uint32 PlatformFormat, uint32 SizeX, uint32 SizeY, uint32 ArraySize, uint32 NumMips, uint32 NumSamples, ETextureCreateFlags CreateFlags)
 {
-	// Need a mutable format so we can reinterpret an sRGB format into a linear format
-	XrSwapchainUsageFlags Usage = XR_SWAPCHAIN_USAGE_MUTABLE_FORMAT_BIT;
+	XrSwapchainUsageFlags Usage = 0;
+	if (!(CreateFlags & TexCreate_SRGB))
+	{
+		// On Windows both sRGB and non-sRGB integer formats have gamma correction, so since we
+		// do gamma correction ourselves in the post-processor we allocate a non-SRGB format.
+		// On OpenXR non-sRGB formats are assumed to be linear without gamma correction,
+		// so we always allocate an sRGB swapchain format. Thus we need to specify the
+		// mutable flag so we can output gamma corrected colors into an sRGB swapchain without
+		// the implicit gamma correction. On mobile platforms the TexCreate_SRGB flag is specified
+		// which indicates the post-processor is disabled and we do need implicit gamma correction.
+		// We skip setting this flag on those platforms as it would incur a large performance hit.
+		Usage |= XR_SWAPCHAIN_USAGE_MUTABLE_FORMAT_BIT;
+	}
 	if (EnumHasAnyFlags(CreateFlags, TexCreate_RenderTargetable))
 	{
 		Usage |= XR_SWAPCHAIN_USAGE_COLOR_ATTACHMENT_BIT;
