Log OpenXR runtime name and version
Also added some verbose logs to the OpenXR plugin

diff --git a/Engine/Plugins/Runtime/OpenXR/Source/OpenXRHMD/Private/OpenXRHMD.cpp b/Engine/Plugins/Runtime/OpenXR/Source/OpenXRHMD/Private/OpenXRHMD.cpp
--- a/Engine/Plugins/Runtime/OpenXR/Source/OpenXRHMD/Private/OpenXRHMD.cpp
+++ b/Engine/Plugins/Runtime/OpenXR/Source/OpenXRHMD/Private/OpenXRHMD.cpp
@@ -2063,9 +2063,7 @@ bool FOpenXRHMD::OnStartGameFrame(FWorldContext& WorldContext)
 
 			CurrentSessionState = SessionState.state;
 
-#if 0
-			UE_LOG(LogHMD, Log, TEXT("Session state switching to %s"), OpenXRSessionStateToString(CurrentSessionState));
-#endif
+			UE_LOG(LogHMD, Verbose, TEXT("Session state switching to %s"), OpenXRSessionStateToString(CurrentSessionState));
 
 			if (SessionState.state == XR_SESSION_STATE_READY)
 			{
@@ -2193,7 +2191,7 @@ void FOpenXRHMD::OnBeginRendering_RHIThread(const FPipelinedFrameState& InFrameS
 	ensure(IsInRenderingThread() || IsInRHIThread());
 
 	// TODO: Add a hook to resolve discarded frames before we start a new frame.
-	checkSlow(!bIsRendering);
+	UE_CLOG(!bIsRendering, LogHMD, Verbose, TEXT("Discarded previous frame and started rendering a new frame."));
 
 	SCOPED_NAMED_EVENT(BeginFrame, FColor::Red);
 
diff --git a/Engine/Plugins/Runtime/OpenXR/Source/OpenXRHMD/Private/OpenXRHMDModule.cpp b/Engine/Plugins/Runtime/OpenXR/Source/OpenXRHMD/Private/OpenXRHMDModule.cpp
--- a/Engine/Plugins/Runtime/OpenXR/Source/OpenXRHMD/Private/OpenXRHMDModule.cpp
+++ b/Engine/Plugins/Runtime/OpenXR/Source/OpenXRHMD/Private/OpenXRHMDModule.cpp
@@ -656,6 +656,12 @@ bool FOpenXRHMDModule::InitInstance()
 		return false;
 	}
 
+	XrInstanceProperties InstanceProps;
+	InstanceProps.type = XR_TYPE_INSTANCE_PROPERTIES;
+	InstanceProps.next = nullptr;
+	XR_ENSURE(xrGetInstanceProperties(Instance, &InstanceProps));
+	UE_LOG(LogHMD, Log, TEXT("Initialized OpenXR on %S runtime version %d.%d.%d"), InstanceProps.runtimeName, XR_VERSION_MAJOR(InstanceProps.runtimeVersion), XR_VERSION_MINOR(InstanceProps.runtimeVersion), XR_VERSION_PATCH(InstanceProps.runtimeVersion));
+
 	for (IOpenXRExtensionPlugin* Module : ExtensionPlugins)
 	{
 		Module->PostCreateInstance(Instance);
