Add XR_EXT_dpad_binding support to OpenXR

No binding modifications supported yet

diff --git a/Engine/Plugins/Runtime/OpenXR/Source/OpenXRHMD/Private/OpenXRHMDModule.cpp b/Engine/Plugins/Runtime/OpenXR/Source/OpenXRHMD/Private/OpenXRHMDModule.cpp
--- a/Engine/Plugins/Runtime/OpenXR/Source/OpenXRHMD/Private/OpenXRHMDModule.cpp
+++ b/Engine/Plugins/Runtime/OpenXR/Source/OpenXRHMD/Private/OpenXRHMDModule.cpp
@@ -439,6 +439,9 @@ bool FOpenXRHMDModule::GetOptionalExtensions(TArray<const ANSICHAR*>& OutExtensi
 	OutExtensions.Add(XR_KHR_COMPOSITION_LAYER_DEPTH_EXTENSION_NAME);
 	OutExtensions.Add(XR_VARJO_QUAD_VIEWS_EXTENSION_NAME);
 	OutExtensions.Add(XR_KHR_VISIBILITY_MASK_EXTENSION_NAME);
+
+	// Draft extension not yet provided in headers
+	OutExtensions.Add("XR_EXT_dpad_binding");
 	return true;
 }
 
diff --git a/Engine/Plugins/Runtime/OpenXR/Source/OpenXRInput/Private/OpenXRInput.cpp b/Engine/Plugins/Runtime/OpenXR/Source/OpenXRInput/Private/OpenXRInput.cpp
--- a/Engine/Plugins/Runtime/OpenXR/Source/OpenXRInput/Private/OpenXRInput.cpp
+++ b/Engine/Plugins/Runtime/OpenXR/Source/OpenXRInput/Private/OpenXRInput.cpp
@@ -169,6 +169,7 @@ FOpenXRInputPlugin::FOpenXRInput::FOpenXRInput(FOpenXRHMD* HMD)
 	, Actions()
 	, Controllers()
 	, bActionsBound(false)
+	, bDirectionalBindingSupported(false)
 	, MessageHandler(new FGenericApplicationMessageHandler())
 {
 	IModularFeatures::Get().RegisterModularFeature(GetModularFeatureName(), this);
@@ -176,6 +177,8 @@ FOpenXRInputPlugin::FOpenXRInput::FOpenXRInput(FOpenXRHMD* HMD)
 	// If there is no HMD then this module is not active, but it still needs to exist so we can EnumerateMotionSources from it.
 	if (OpenXRHMD)
 	{
+		bDirectionalBindingSupported = OpenXRHMD->IsExtensionEnabled("XR_EXT_dpad_binding");
+
 		// Note: AnyHand needs special handling because it tries left then falls back to right in each call.
 		MotionSourceToControllerHandMap.Add(OpenXRSourceNames::Left, EControllerHand::Left);
 		MotionSourceToControllerHandMap.Add(OpenXRSourceNames::Right, EControllerHand::Right);
@@ -248,7 +251,6 @@ void FOpenXRInputPlugin::FOpenXRInput::BuildActions()
 	{
 		return;
 	}
-	
 
 	XrInstance Instance = OpenXRHMD->GetInstance();
 	check(Instance);
@@ -497,10 +499,9 @@ int32 FOpenXRInputPlugin::FOpenXRInput::SuggestBindings(XrInstance Instance, FOp
 			{
 				Path += "/touchvalue";  // Note: this is not a standard openxr identifier.  It is meant to represent some kind of analog touch sensor.
 			}
-			else if (Component == "up" || Component == "down" || Component == "left" || Component == "right")
+			else if (bDirectionalBindingSupported && (Component == "up" || Component == "down" || Component == "left" || Component == "right"))
 			{
-				// TODO: Reserved for D-Pad support
-				continue;
+				Path += "/dpad_" + Component;
 			}
 			else
 			{
diff --git a/Engine/Plugins/Runtime/OpenXR/Source/OpenXRInput/Private/OpenXRInput.h b/Engine/Plugins/Runtime/OpenXR/Source/OpenXRInput/Private/OpenXRInput.h
--- a/Engine/Plugins/Runtime/OpenXR/Source/OpenXRInput/Private/OpenXRInput.h
+++ b/Engine/Plugins/Runtime/OpenXR/Source/OpenXRInput/Private/OpenXRInput.h
@@ -103,7 +103,9 @@ public:
 		int32 GetDeviceIDForMotionSource(FName MotionSource) const;
 		XrPath GetUserPathForMotionSource(FName MotionSource) const;
 		bool IsOpenXRInputSupportedMotionSource(const FName MotionSource) const;
+
 		bool bActionsBound;
+		bool bDirectionalBindingSupported;
 
 		void BuildActions();
 		void DestroyActions();
