Fix redundant acquire on main swapchain

diff --git a/Engine/Plugins/Runtime/OpenXR/Source/OpenXRHMD/Private/OpenXRHMD.cpp b/Engine/Plugins/Runtime/OpenXR/Source/OpenXRHMD/Private/OpenXRHMD.cpp
--- a/Engine/Plugins/Runtime/OpenXR/Source/OpenXRHMD/Private/OpenXRHMD.cpp
+++ b/Engine/Plugins/Runtime/OpenXR/Source/OpenXRHMD/Private/OpenXRHMD.cpp
@@ -1681,8 +1681,12 @@ bool FOpenXRHMD::AllocateRenderTargetTexture(uint32 Index, uint32 SizeX, uint32
 			return false;
 		}
 
-		// Acquire the first swapchain image
-		Swapchain->IncrementSwapChainIndex_RHIThread();
+		// Acquire the first swapchain image if we're already rendering since the old swapchain was already acquired.
+		// We need to be careful here that we only re-acquire swapchains and not acquire swapchain images redundantly.
+		if (bIsRendering)
+		{
+			Swapchain->IncrementSwapChainIndex_RHIThread();
+		}
 
 #if WITH_EDITOR
 		if (GIsEditor)
@@ -1720,6 +1724,14 @@ bool FOpenXRHMD::AllocateDepthTexture(uint32 Index, uint32 SizeX, uint32 SizeY,
 		return false;
 	}
 
+	// Ensure the texture size matches the eye layer. We may get other depth allocations unrelated to the main scene render.
+	// FIXME: This introduces a magic number that will break this logic. A developer can select a scene capture target that
+	// happens to be the same size as the ideal render target and break everything for reasons that will be a mystery to them.
+	if (FIntPoint(SizeX, SizeY) != GetIdealRenderTargetSize())
+	{
+		return false;
+	}
+
 	// We're only creating a 1x target here, but we don't know whether it'll be the targeted texture
 	// or the resolve texture. Because of this, we unify the input flags.
 	ETextureCreateFlags UnifiedCreateFlags = Flags | TargetableTextureFlags;
@@ -1745,8 +1757,12 @@ bool FOpenXRHMD::AllocateDepthTexture(uint32 Index, uint32 SizeX, uint32 SizeY,
 			return false;
 		}
 
-		// Acquire the first swapchain image
-		Swapchain->IncrementSwapChainIndex_RHIThread();
+		// Acquire the first swapchain image if we're already rendering since the old swapchain was already acquired.
+		// We need to be careful here that we only re-acquire swapchains and not acquire swapchain images redundantly.
+		if (bIsRendering)
+		{
+			Swapchain->IncrementSwapChainIndex_RHIThread();
+		}
 	}
 
 	bNeedReAllocatedDepth = false;
@@ -1762,6 +1778,9 @@ void FOpenXRHMD::OnBeginRendering_RenderThread(FRHICommandListImmediate& RHICmdL
 	ensure(IsInRenderingThread());
 	FReadScopeLock DeviceLock(DeviceMutex);
 
+	UE_CLOG(PipelinedFrameStateRendering.FrameState.predictedDisplayTime >= PipelinedFrameStateGame.FrameState.predictedDisplayTime,
+		LogHMD, VeryVerbose, TEXT("Predicted display time went backwards from %lld to %lld"), PipelinedFrameStateRendering.FrameState.predictedDisplayTime, PipelinedFrameStateGame.FrameState.predictedDisplayTime);
+
 	for (IOpenXRExtensionPlugin* Module : ExtensionPlugins)
 	{
 		Module->OnBeginRendering_RenderThread(Session);
@@ -2191,7 +2210,7 @@ void FOpenXRHMD::OnBeginRendering_RHIThread(const FPipelinedFrameState& InFrameS
 	ensure(IsInRenderingThread() || IsInRHIThread());
 
 	// TODO: Add a hook to resolve discarded frames before we start a new frame.
-	UE_CLOG(!bIsRendering, LogHMD, Verbose, TEXT("Discarded previous frame and started rendering a new frame."));
+	UE_CLOG(bIsRendering, LogHMD, Verbose, TEXT("Discarded previous frame and started rendering a new frame."));
 
 	SCOPED_NAMED_EVENT(BeginFrame, FColor::Red);
 
@@ -2235,6 +2254,8 @@ void FOpenXRHMD::OnBeginRendering_RHIThread(const FPipelinedFrameState& InFrameS
 		}
 
 		bIsRendering = true;
+
+		UE_LOG(LogHMD, VeryVerbose, TEXT("Rendering frame predicted to be displayed at %lld"), InFrameState.FrameState.predictedDisplayTime);
 	}
 	else
 	{
@@ -2322,6 +2343,9 @@ void FOpenXRHMD::OnFinishRendering_RHIThread()
 			}
 			EndInfo.next = Module->OnEndFrame(Session, EndInfo.displayTime, ColorImages, DepthImages, EndInfo.next);
 		}
+
+		UE_LOG(LogHMD, VeryVerbose, TEXT("Presenting frame predicted to be displayed at %lld"), PipelinedFrameStateRHI.FrameState.predictedDisplayTime);
+
 		XR_ENSURE(xrEndFrame(Session, &EndInfo));
 	}
 
diff --git a/Engine/Plugins/Runtime/OpenXR/Source/OpenXRHMD/Private/OpenXRHMD_Swapchain.cpp b/Engine/Plugins/Runtime/OpenXR/Source/OpenXRHMD/Private/OpenXRHMD_Swapchain.cpp
--- a/Engine/Plugins/Runtime/OpenXR/Source/OpenXRHMD/Private/OpenXRHMD_Swapchain.cpp
+++ b/Engine/Plugins/Runtime/OpenXR/Source/OpenXRHMD/Private/OpenXRHMD_Swapchain.cpp
@@ -12,7 +12,8 @@ static TAutoConsoleVariable<int32> CVarOpenXRSwapchainRetryCount(
 
 FOpenXRSwapchain::FOpenXRSwapchain(TArray<FTextureRHIRef>&& InRHITextureSwapChain, const FTextureRHIRef & InRHITexture, XrSwapchain InHandle) :
 	FXRSwapChain(MoveTemp(InRHITextureSwapChain), InRHITexture),
-	Handle(InHandle)
+	Handle(InHandle),
+	Acquired(false)
 {
 }
 
@@ -21,12 +22,22 @@ FOpenXRSwapchain::~FOpenXRSwapchain()
 	XR_ENSURE(xrDestroySwapchain(Handle));
 }
 
-// TODO: This function should be renamed to IncrementSwapChainIndex_RenderThread.
-// Name change is currently blocked on runtimes still requiring this on the RHI thread.
+// TODO: This function should be moved to the RenderThread, that change is currently blocked
+// by the Vulkan extension requiring access to the VkQueue in xrAcquireSwapchainImage.
 void FOpenXRSwapchain::IncrementSwapChainIndex_RHIThread()
 {
 	check(IsInRenderingThread() || IsInRHIThread());
 
+	// TODO: When moving this function to the RenderThread remove this logic so the RenderThread
+	// can acquire a new swapchain image before the RHI thread is done with it.
+	bool WasAcquired = false;
+	Acquired.compare_exchange_strong(WasAcquired, true);
+	if (WasAcquired)
+	{
+		UE_LOG(LogHMD, Verbose, TEXT("Attempted to redundantly acquire image %d in swapchain %p"), SwapChainIndex_RHIThread, Handle);
+		return;
+	}
+
 	SCOPED_NAMED_EVENT(AcquireImage, FColor::Red);
 
 	XrSwapchainImageAcquireInfo Info;
@@ -34,6 +45,8 @@ void FOpenXRSwapchain::IncrementSwapChainIndex_RHIThread()
 	Info.next = nullptr;
 	XR_ENSURE(xrAcquireSwapchainImage(Handle, &Info, &SwapChainIndex_RHIThread));
 
+	UE_LOG(LogHMD, VeryVerbose, TEXT("Acquired image %d in swapchain %p"), SwapChainIndex_RHIThread, Handle);
+
 	GDynamicRHI->RHIAliasTextureResources((FTextureRHIRef&)RHITexture, (FTextureRHIRef&)RHITextureSwapChain[SwapChainIndex_RHIThread]);
 }
 
@@ -41,6 +54,12 @@ void FOpenXRSwapchain::WaitCurrentImage_RHIThread(int64 Timeout)
 {
 	check(IsInRenderingThread() || IsInRHIThread());
 
+	if (!Acquired)
+	{
+		UE_LOG(LogHMD, Verbose, TEXT("Attempted to wait on unacquired image %d in swapchain %p"), SwapChainIndex_RHIThread, Handle);
+		return;
+	}
+
 	SCOPED_NAMED_EVENT(WaitImage, FColor::Red);
 
 	XrSwapchainImageWaitInfo WaitInfo;
@@ -64,18 +83,30 @@ void FOpenXRSwapchain::WaitCurrentImage_RHIThread(int64 Timeout)
 		// We can't continue without acquiring a new swapchain image since we won't have an image available to render to.
 		UE_LOG(LogHMD, Fatal, TEXT("Failed to wait on acquired swapchain image. This usually indicates a problem with the OpenXR runtime."));
 	}
+
+	UE_LOG(LogHMD, VeryVerbose, TEXT("Waited on image %d in swapchain %p"), SwapChainIndex_RHIThread, Handle);
 }
 
 void FOpenXRSwapchain::ReleaseCurrentImage_RHIThread()
 {
 	check(IsInRenderingThread() || IsInRHIThread());
 
+	bool WasAcquired = true;
+	Acquired.compare_exchange_strong(WasAcquired, false);
+	if (!WasAcquired)
+	{
+		UE_LOG(LogHMD, Verbose, TEXT("Attempted to release unacquired image %d in swapchain %p"), SwapChainIndex_RHIThread, Handle);
+		return;
+	}
+
 	SCOPED_NAMED_EVENT(ReleaseImage, FColor::Red);
 
 	XrSwapchainImageReleaseInfo ReleaseInfo;
 	ReleaseInfo.type = XR_TYPE_SWAPCHAIN_IMAGE_RELEASE_INFO;
 	ReleaseInfo.next = nullptr;
 	XR_ENSURE(xrReleaseSwapchainImage(Handle, &ReleaseInfo));
+
+	UE_LOG(LogHMD, VeryVerbose, TEXT("Released on image %d in swapchain %p"), SwapChainIndex_RHIThread, Handle);
 }
 
 uint8 FOpenXRSwapchain::GetNearestSupportedSwapchainFormat(XrSession InSession, uint8 RequestedFormat, TFunction<uint32(uint8)> ToPlatformFormat /*= nullptr*/)
