Move xrWaitFrame to the OnStartGameFrame callback

This gives us less simulation headroom, but prevents a pipeline stall in UE5
Also includes some UnrealInsights markers for the Oculus plugin

diff --git a/Engine/Plugins/Runtime/Oculus/OculusVR/Source/OculusHMD/Private/OculusHMD.cpp b/Engine/Plugins/Runtime/Oculus/OculusVR/Source/OculusHMD/Private/OculusHMD.cpp
--- a/Engine/Plugins/Runtime/Oculus/OculusVR/Source/OculusHMD/Private/OculusHMD.cpp
+++ b/Engine/Plugins/Runtime/Oculus/OculusVR/Source/OculusHMD/Private/OculusHMD.cpp
@@ -3289,6 +3289,8 @@ namespace OculusHMD
 				{
 					if (FOculusHMDModule::GetPluginWrapper().GetInitialized() && WaitFrameNumber != Frame->FrameNumber)
 					{
+						SCOPED_NAMED_EVENT(WaitFrame, FColor::Red);
+
 						UE_LOG(LogHMD, Verbose, TEXT("FOculusHMDModule::GetPluginWrapper().WaitToBeginFrame %u"), Frame->FrameNumber);
 
 						ovrpResult Result;
@@ -3489,6 +3491,8 @@ namespace OculusHMD
 					FOculusHMDModule::GetPluginWrapper().GetNativeXrApiType(&NativeXrApi);
 					if ((Frame_RHIThread->ShowFlags.Rendering || NativeXrApi == ovrpXrApi_OpenXR) && !Frame_RHIThread->Flags.bSplashIsShown)
 					{
+						SCOPED_NAMED_EVENT(BeginFrame, FColor::Red);
+
 						UE_LOG(LogHMD, Verbose, TEXT("FOculusHMDModule::GetPluginWrapper().BeginFrame4 %u"), Frame_RHIThread->FrameNumber);
 
 						ovrpResult Result;
@@ -3523,6 +3527,8 @@ namespace OculusHMD
 			FOculusHMDModule::GetPluginWrapper().GetNativeXrApiType(&NativeXrApi);
 			if ((Frame_RHIThread->ShowFlags.Rendering || NativeXrApi == ovrpXrApi_OpenXR) && !Frame_RHIThread->Flags.bSplashIsShown)
 			{
+				SCOPED_NAMED_EVENT(EndFrame, FColor::Red);
+
 				TArray<FLayerPtr> Layers = Layers_RHIThread;
 				Layers.Sort(FLayerPtr_CompareTotal());
 				TArray<const ovrpLayerSubmit*> LayerSubmitPtr;
diff --git a/Engine/Plugins/Runtime/OpenXR/Source/OpenXRHMD/Private/OpenXRHMD.cpp b/Engine/Plugins/Runtime/OpenXR/Source/OpenXRHMD/Private/OpenXRHMD.cpp
--- a/Engine/Plugins/Runtime/OpenXR/Source/OpenXRHMD/Private/OpenXRHMD.cpp
+++ b/Engine/Plugins/Runtime/OpenXR/Source/OpenXRHMD/Private/OpenXRHMD.cpp
@@ -873,6 +873,7 @@ FOpenXRHMD::FOpenXRHMD(const FAutoRegister& AutoRegister, XrInstance InInstance,
 	, bIsReady(false)
 	, bIsRendering(false)
 	, bIsSynchronized(false)
+	, bShouldWait(true)
 	, bIsExitingSessionByxrRequestExitSession(false)
 	, bNeedReAllocatedDepth(false)
 	, bNeedReBuildOcclusionMesh(true)
@@ -2003,9 +2004,18 @@ void FOpenXRHMD::OnLateUpdateApplied_RenderThread(FRHICommandListImmediate& RHIC
 }
 
 void FOpenXRHMD::OnBeginRendering_GameThread()
+{
+	// We need to make sure we keep the Wait/Begin/End triplet in sync, so here we signal that we
+	// can wait for the next frame in the next tick. Without this signal it's possible that two ticks
+	// happen before the next frame is actually rendered.
+	bShouldWait = true;
+}
+
+void FOpenXRHMD::OnBeginSimulation_GameThread()
 {
 	FReadScopeLock Lock(SessionHandleMutex);
-	if (!bIsReady || !bIsRunning)
+
+	if (!bIsReady || !bIsRunning || !bShouldWait)
 	{
 		// @todo: Sleep here?
 		return;
@@ -2034,6 +2044,8 @@ void FOpenXRHMD::OnBeginRendering_GameThread()
 	PipelineState.TrackingSpace = GetTrackingSpace();
 	PipelineState.WorldToMetersScale = WorldToMetersScale;
 
+	bShouldWait = false;
+
 	EnumerateViews(PipelineState);
 }
 
@@ -2173,9 +2185,10 @@ bool FOpenXRHMD::OnStartGameFrame(FWorldContext& WorldContext)
 
 	GetARCompositionComponent()->StartARGameFrame(WorldContext);
 
-	// Add a display period to the simulation frame state so we're predicting poses for the new frame.
-	FPipelinedFrameState& PipelineState = GetPipelinedFrameStateForThread();
-	PipelineState.FrameState.predictedDisplayTime += PipelineState.FrameState.predictedDisplayPeriod;
+	// TODO: We could do this earlier in the pipeline and allow simulation to run one frame ahead of the render thread.
+	// That would allow us to take more advantage of Late Update and give projects more headroom for simulation.
+	// However currently blocking in earlier callbacks can result in a pipeline stall, so we do it here instead.
+	OnBeginSimulation_GameThread();
 
 	// Snapshot new poses for game simulation.
 	UpdateDeviceLocations(true);
diff --git a/Engine/Plugins/Runtime/OpenXR/Source/OpenXRHMD/Private/OpenXRHMD.h b/Engine/Plugins/Runtime/OpenXR/Source/OpenXRHMD/Private/OpenXRHMD.h
--- a/Engine/Plugins/Runtime/OpenXR/Source/OpenXRHMD/Private/OpenXRHMD.h
+++ b/Engine/Plugins/Runtime/OpenXR/Source/OpenXRHMD/Private/OpenXRHMD.h
@@ -279,6 +279,7 @@ public:
 	/** Destructor */
 	virtual ~FOpenXRHMD();
 
+	void OnBeginSimulation_GameThread();
 	void OnBeginRendering_RHIThread(const FPipelinedFrameState& InFrameState, FXRSwapChainPtr ColorSwapchain, FXRSwapChainPtr DepthSwapchain);
 	void OnFinishRendering_RHIThread();
 
@@ -315,6 +316,7 @@ private:
 	TAtomic<bool>			bIsReady;
 	TAtomic<bool>			bIsRendering;
 	TAtomic<bool>			bIsSynchronized;
+	bool					bShouldWait;
 	bool					bIsExitingSessionByxrRequestExitSession;
 	bool					bDepthExtensionSupported;
 	bool					bHiddenAreaMaskSupported;
