Prevent using a swapchain image that wasn't waited for (UE-141829).

diff --git a/Engine/Plugins/Runtime/OpenXR/Source/OpenXRHMD/Private/OpenXRHMD.cpp b/Engine/Plugins/Runtime/OpenXR/Source/OpenXRHMD/Private/OpenXRHMD.cpp
--- a/Engine/Plugins/Runtime/OpenXR/Source/OpenXRHMD/Private/OpenXRHMD.cpp
+++ b/Engine/Plugins/Runtime/OpenXR/Source/OpenXRHMD/Private/OpenXRHMD.cpp
@@ -1682,12 +1682,7 @@ bool FOpenXRHMD::AllocateRenderTargetTexture(uint32 Index, uint32 SizeX, uint32
 			return false;
 		}
 
-		// Acquire the first swapchain image if we're already rendering since the old swapchain was already acquired.
-		// We need to be careful here that we only re-acquire swapchains and not acquire swapchain images redundantly.
-		if (bIsRendering)
-		{
-			Swapchain->IncrementSwapChainIndex_RHIThread();
-		}
+		// image will be acquired next time we begin the rendering
 
 #if WITH_EDITOR
 		if (GIsEditor)
@@ -1758,12 +1753,7 @@ bool FOpenXRHMD::AllocateDepthTexture(uint32 Index, uint32 SizeX, uint32 SizeY,
 			return false;
 		}
 
-		// Acquire the first swapchain image if we're already rendering since the old swapchain was already acquired.
-		// We need to be careful here that we only re-acquire swapchains and not acquire swapchain images redundantly.
-		if (bIsRendering)
-		{
-			Swapchain->IncrementSwapChainIndex_RHIThread();
-		}
+		// image will be acquired next time we begin the rendering
 	}
 
 	bNeedReAllocatedDepth = false;
diff --git a/Engine/Plugins/Runtime/OpenXR/Source/OpenXRHMD/Private/OpenXRHMD_Swapchain.cpp b/Engine/Plugins/Runtime/OpenXR/Source/OpenXRHMD/Private/OpenXRHMD_Swapchain.cpp
--- a/Engine/Plugins/Runtime/OpenXR/Source/OpenXRHMD/Private/OpenXRHMD_Swapchain.cpp
+++ b/Engine/Plugins/Runtime/OpenXR/Source/OpenXRHMD/Private/OpenXRHMD_Swapchain.cpp
@@ -13,7 +13,8 @@ static TAutoConsoleVariable<int32> CVarOpenXRSwapchainRetryCount(
 FOpenXRSwapchain::FOpenXRSwapchain(TArray<FTextureRHIRef>&& InRHITextureSwapChain, const FTextureRHIRef & InRHITexture, XrSwapchain InHandle) :
 	FXRSwapChain(MoveTemp(InRHITextureSwapChain), InRHITexture),
 	Handle(InHandle),
-	Acquired(false)
+	ImageAcquired(false),
+	ImageReady(false)
 {
 }
 
@@ -31,7 +32,7 @@ void FOpenXRSwapchain::IncrementSwapChainIndex_RHIThread()
 	// TODO: When moving this function to the RenderThread remove this logic so the RenderThread
 	// can acquire a new swapchain image before the RHI thread is done with it.
 	bool WasAcquired = false;
-	Acquired.compare_exchange_strong(WasAcquired, true);
+	ImageAcquired.compare_exchange_strong(WasAcquired, true);
 	if (WasAcquired)
 	{
 		UE_LOG(LogHMD, Verbose, TEXT("Attempted to redundantly acquire image %d in swapchain %p"), SwapChainIndex_RHIThread, Handle);
@@ -54,9 +55,17 @@ void FOpenXRSwapchain::WaitCurrentImage_RHIThread(int64 Timeout)
 {
 	check(IsInRenderingThread() || IsInRHIThread());
 
-	if (!Acquired)
+	if (!ImageAcquired)
 	{
-		UE_LOG(LogHMD, Verbose, TEXT("Attempted to wait on unacquired image %d in swapchain %p"), SwapChainIndex_RHIThread, Handle);
+		UE_LOG(LogHMD, Warning, TEXT("Attempted to wait on unacquired image %d in swapchain %p"), SwapChainIndex_RHIThread, Handle);
+		return;
+	}
+
+	bool WasReady = false;
+	ImageReady.compare_exchange_strong(WasReady, true);
+	if (WasReady)
+	{
+		UE_LOG(LogHMD, Verbose, TEXT("Attempted to redundantly wait on image %d in swapchain %p"), SwapChainIndex_RHIThread, Handle);
 		return;
 	}
 
@@ -92,10 +101,18 @@ void FOpenXRSwapchain::ReleaseCurrentImage_RHIThread()
 	check(IsInRenderingThread() || IsInRHIThread());
 
 	bool WasAcquired = true;
-	Acquired.compare_exchange_strong(WasAcquired, false);
+	ImageAcquired.compare_exchange_strong(WasAcquired, false);
 	if (!WasAcquired)
 	{
-		UE_LOG(LogHMD, Verbose, TEXT("Attempted to release unacquired image %d in swapchain %p"), SwapChainIndex_RHIThread, Handle);
+		UE_LOG(LogHMD, Warning, TEXT("Attempted to release unacquired image %d in swapchain %p"), SwapChainIndex_RHIThread, Handle);
+		return;
+	}
+
+	bool WasReady = true;
+	ImageReady.compare_exchange_strong(WasAcquired, false);
+	if (!WasReady)
+	{
+		UE_LOG(LogHMD, Warning, TEXT("Attempted to release image %d in swapchain %p that wasn't ready for being written to."), SwapChainIndex_RHIThread, Handle);
 		return;
 	}
 
diff --git a/Engine/Plugins/Runtime/OpenXR/Source/OpenXRHMD/Private/OpenXRHMD_Swapchain.h b/Engine/Plugins/Runtime/OpenXR/Source/OpenXRHMD/Private/OpenXRHMD_Swapchain.h
--- a/Engine/Plugins/Runtime/OpenXR/Source/OpenXRHMD/Private/OpenXRHMD_Swapchain.h
+++ b/Engine/Plugins/Runtime/OpenXR/Source/OpenXRHMD/Private/OpenXRHMD_Swapchain.h
@@ -25,7 +25,12 @@ public:
 
 protected:
 	XrSwapchain Handle;
-	std::atomic<bool> Acquired;
+
+	/** Whether the image associated with the swapchain has been acquired. */
+	std::atomic<bool> ImageAcquired;
+
+	/** Whether the image associated with the swapchain is ready for being written to. */
+	std::atomic<bool> ImageReady;
 };
 
 #ifdef XR_USE_GRAPHICS_API_D3D11
