Fix path to OculusOpenXRLoader UPL file

#android

diff --git a/Engine/Plugins/Runtime/OpenXR/Source/OpenXRHMD/OpenXRHMD.Build.cs b/Engine/Plugins/Runtime/OpenXR/Source/OpenXRHMD/OpenXRHMD.Build.cs
--- a/Engine/Plugins/Runtime/OpenXR/Source/OpenXRHMD/OpenXRHMD.Build.cs
+++ b/Engine/Plugins/Runtime/OpenXR/Source/OpenXRHMD/OpenXRHMD.Build.cs
@@ -113,7 +113,7 @@ namespace UnrealBuildTool.Rules
 			{
 				// If the Oculus plugin is not enabled we need to include our own APL
 				string PluginPath = Utils.MakePathRelativeTo(ModuleDirectory, Target.RelativeEnginePath);
-				AdditionalPropertiesForReceipt.Add("AndroidPlugin", Path.Combine(PluginPath, "OculusOpenXRLoader_APL.xml"));
+				AdditionalPropertiesForReceipt.Add("AndroidPlugin", Path.Combine(PluginPath, "..", "..", "OculusOpenXRLoader_APL.xml"));
 			}
 		}
 	}
