Add XR_KHR_binding_modification extension to OpenXR

Required for the dpad binding extension even though we don't use it

diff --git a/Engine/Plugins/Runtime/OpenXR/Source/OpenXRHMD/Private/OpenXRHMDModule.cpp b/Engine/Plugins/Runtime/OpenXR/Source/OpenXRHMD/Private/OpenXRHMDModule.cpp
--- a/Engine/Plugins/Runtime/OpenXR/Source/OpenXRHMD/Private/OpenXRHMDModule.cpp
+++ b/Engine/Plugins/Runtime/OpenXR/Source/OpenXRHMD/Private/OpenXRHMDModule.cpp
@@ -440,6 +440,7 @@ bool FOpenXRHMDModule::GetOptionalExtensions(TArray<const ANSICHAR*>& OutExtensi
 	OutExtensions.Add(XR_KHR_COMPOSITION_LAYER_DEPTH_EXTENSION_NAME);
 	OutExtensions.Add(XR_VARJO_QUAD_VIEWS_EXTENSION_NAME);
 	OutExtensions.Add(XR_KHR_VISIBILITY_MASK_EXTENSION_NAME);
+	OutExtensions.Add(XR_KHR_BINDING_MODIFICATION_EXTENSION_NAME);
 
 	// Draft extension not yet provided in headers
 	OutExtensions.Add("XR_EXT_dpad_binding");
