Fix failure to set XR input bindings on startup (UE-145485).

diff --git a/Engine/Plugins/Runtime/Oculus/OculusVR/Source/OculusOpenXRHMD/Private/OculusOpenXRHMD.cpp b/Engine/Plugins/Runtime/Oculus/OculusVR/Source/OculusOpenXRHMD/Private/OculusOpenXRHMD.cpp
--- a/Engine/Plugins/Runtime/Oculus/OculusVR/Source/OculusOpenXRHMD/Private/OculusOpenXRHMD.cpp
+++ b/Engine/Plugins/Runtime/Oculus/OculusVR/Source/OculusOpenXRHMD/Private/OculusOpenXRHMD.cpp
@@ -21,7 +21,7 @@ bool FOculusOpenXRHMD::GetRequiredExtensions(TArray<const ANSICHAR*>& OutExtensi
 bool FOculusOpenXRHMD::GetInteractionProfile(XrInstance InInstance, FString& OutKeyPrefix, XrPath& OutPath, bool& OutHasHaptics)
 {
 	//UE_LOG(LogOculusOpenXRPlugin, Log, TEXT("Oculus OpenXR GetInteractionProfile"));
-	return true;
+	return false;	// if you return true, make sure OutPath and OutHasHaptics are initialized
 }
 
 bool FOculusOpenXRHMD::GetSpectatorScreenController(FHeadMountedDisplayBase* InHMDBase, TUniquePtr<FDefaultSpectatorScreenController>& OutSpectatorScreenController)
diff --git a/Engine/Plugins/Runtime/OpenXR/Source/OpenXRInput/Private/OpenXRInput.cpp b/Engine/Plugins/Runtime/OpenXR/Source/OpenXRInput/Private/OpenXRInput.cpp
--- a/Engine/Plugins/Runtime/OpenXR/Source/OpenXRInput/Private/OpenXRInput.cpp
+++ b/Engine/Plugins/Runtime/OpenXR/Source/OpenXRInput/Private/OpenXRInput.cpp
@@ -289,9 +289,9 @@ void FOpenXRInputPlugin::FOpenXRInput::BuildActions()
 	for (IOpenXRExtensionPlugin* Plugin : OpenXRHMD->GetExtensionPlugins())
 	{
 		FString KeyPrefix;
-		XrPath Path;
-		bool HasHaptics;
-		if (Plugin->GetInteractionProfile(Instance, KeyPrefix, Path, HasHaptics))
+		XrPath Path = XR_NULL_PATH;
+		bool HasHaptics = false;
+		if (Plugin->GetInteractionProfile(Instance, KeyPrefix, Path, HasHaptics) && Path != XR_NULL_PATH)
 		{
 			Profiles.Add(KeyPrefix, FInteractionProfile(Path, HasHaptics));
 		}
@@ -499,8 +499,12 @@ int32 FOpenXRInputPlugin::FOpenXRInput::SuggestBindings(XrInstance Instance, FOp
 			{
 				Path += "/touchvalue";  // Note: this is not a standard openxr identifier.  It is meant to represent some kind of analog touch sensor.
 			}
-			else if (bDirectionalBindingSupported && (Component == "up" || Component == "down" || Component == "left" || Component == "right"))
+			else if (Component == "up" || Component == "down" || Component == "left" || Component == "right")
 			{
+				if (!bDirectionalBindingSupported)
+				{
+					continue;
+				}
 				Path += "/dpad_" + Component;
 			}
 			else
