OpenXR Oculus: Properly detect Oculus plugin to facilitate APL selection

We were using the wrong facility in UBT to detect the active project plugins.
Now we use ProjectDescriptor + PluginReferenceDescriptor to verify the Oculus plugin existence.

diff --git a/Engine/Plugins/Runtime/OpenXR/Source/OpenXRHMD/OpenXRHMD.Build.cs b/Engine/Plugins/Runtime/OpenXR/Source/OpenXRHMD/OpenXRHMD.Build.cs
--- a/Engine/Plugins/Runtime/OpenXR/Source/OpenXRHMD/OpenXRHMD.Build.cs
+++ b/Engine/Plugins/Runtime/OpenXR/Source/OpenXRHMD/OpenXRHMD.Build.cs
@@ -2,6 +2,7 @@
 
 using System.IO;
 using System.Linq;
+using EpicGames.Core;
 
 namespace UnrealBuildTool.Rules
 {
@@ -82,11 +83,34 @@ namespace UnrealBuildTool.Rules
                 AddEngineThirdPartyPrivateStaticDependencies(Target, "Vulkan");
 			}
 
-			if (Target.Platform == UnrealTargetPlatform.Android && !Target.EnablePlugins.Contains("OculusVR"))
+			if (Target.Platform == UnrealTargetPlatform.Android)
 			{
-				// If the Oculus plugin is not enabled we need to include our own APL
-				string PluginPath = Utils.MakePathRelativeTo(ModuleDirectory, Target.RelativeEnginePath);
-				AdditionalPropertiesForReceipt.Add("AndroidPlugin", Path.Combine(PluginPath, "..", "..", "OculusOpenXRLoader_APL.xml"));
+				bool bAndroidOculusEnabled = false;
+				if (Target.ProjectFile != null)
+				{
+					ProjectDescriptor Project = ProjectDescriptor.FromFile(Target.ProjectFile);
+					if (Project.Plugins != null)
+					{
+						foreach (PluginReferenceDescriptor PluginDescriptor in Project.Plugins)
+						{
+							if ((PluginDescriptor.Name == "OculusVR") && PluginDescriptor.IsEnabledForPlatform(Target.Platform))
+							{
+								Log.TraceInformation("OpenXRHMD: Android Oculus plugin enabled, will use OculusMobile_APL.xml");
+								bAndroidOculusEnabled = true;
+								break;
+							}
+						}
+					}
+				}
+
+				if (!bAndroidOculusEnabled)
+				{
+					Log.TraceInformation("OpenXRHMD: Using OculusOpenXRLoader_APL.xml");
+
+					// If the Oculus plugin is not enabled we need to include our own APL
+					string PluginPath = Utils.MakePathRelativeTo(ModuleDirectory, Target.RelativeEnginePath);
+					AdditionalPropertiesForReceipt.Add("AndroidPlugin", Path.Combine(PluginPath, "..", "..", "OculusOpenXRLoader_APL.xml"));
+				}
 			}
 		}
 	}
