Removing more includes of private RHI headers.
Removing unnecessary dependencies on RHICore and  RHI ThirdParty libraries.

diff --git a/Engine/Plugins/Media/BlackmagicMedia/Source/BlackmagicMediaOutput/BlackmagicMediaOutput.Build.cs b/Engine/Plugins/Media/BlackmagicMedia/Source/BlackmagicMediaOutput/BlackmagicMediaOutput.Build.cs
--- a/Engine/Plugins/Media/BlackmagicMedia/Source/BlackmagicMediaOutput/BlackmagicMediaOutput.Build.cs
+++ b/Engine/Plugins/Media/BlackmagicMedia/Source/BlackmagicMediaOutput/BlackmagicMediaOutput.Build.cs
@@ -30,7 +30,6 @@ namespace UnrealBuildTool.Rules
 					"Engine",
 					"RenderCore",
 					"RHI",
-					"RHICore",
 					"Slate"
 				}
             );
diff --git a/Engine/Plugins/Media/ElectraUtil/Source/ElectraSamples/ElectraSamples.Build.cs b/Engine/Plugins/Media/ElectraUtil/Source/ElectraSamples/ElectraSamples.Build.cs
--- a/Engine/Plugins/Media/ElectraUtil/Source/ElectraSamples/ElectraSamples.Build.cs
+++ b/Engine/Plugins/Media/ElectraUtil/Source/ElectraSamples/ElectraSamples.Build.cs
@@ -24,7 +24,6 @@ namespace UnrealBuildTool.Rules
 					"MediaUtils",
 					"RenderCore",
 					"RHI",
-					"RHICore",
 					"ElectraBase",
 				});
 
diff --git a/Engine/Plugins/Runtime/OpenXR/Source/OpenXRHMD/OpenXRHMD.Build.cs b/Engine/Plugins/Runtime/OpenXR/Source/OpenXRHMD/OpenXRHMD.Build.cs
--- a/Engine/Plugins/Runtime/OpenXR/Source/OpenXRHMD/OpenXRHMD.Build.cs
+++ b/Engine/Plugins/Runtime/OpenXR/Source/OpenXRHMD/OpenXRHMD.Build.cs
@@ -59,16 +59,9 @@ namespace UnrealBuildTool.Rules
                 AddEngineThirdPartyPrivateStaticDependencies(Target, "DX11", "DX12");
             }
 
-			if (Target.IsInPlatformGroup(UnrealPlatformGroup.Windows) || Target.IsInPlatformGroup(UnrealPlatformGroup.Linux))
-			{
-				AddEngineThirdPartyPrivateStaticDependencies(Target, "NVAftermath");
-			}
-
             if (Target.Platform == UnrealTargetPlatform.Win64 || Target.Platform == UnrealTargetPlatform.Android)
             {
-                PrivateDependencyModuleNames.AddRange(new string[] {
-                    "OpenGLDrv",
-                });
+                PrivateDependencyModuleNames.Add("OpenGLDrv");
 
                 AddEngineThirdPartyPrivateStaticDependencies(Target, "OpenGL");
 			}
@@ -76,9 +69,7 @@ namespace UnrealBuildTool.Rules
 			if (Target.Platform == UnrealTargetPlatform.Win64 || Target.Platform == UnrealTargetPlatform.Android  
 			    || Target.IsInPlatformGroup(UnrealPlatformGroup.Linux))
             {
-                PrivateDependencyModuleNames.AddRange(new string[] {
-                    "VulkanRHI"
-                });
+                PrivateDependencyModuleNames.Add("VulkanRHI");
 
                 AddEngineThirdPartyPrivateStaticDependencies(Target, "Vulkan");
 			}
diff --git a/Engine/Plugins/Runtime/Steam/SteamVR/Source/SteamVR/SteamVR.Build.cs b/Engine/Plugins/Runtime/Steam/SteamVR/Source/SteamVR/SteamVR.Build.cs
--- a/Engine/Plugins/Runtime/Steam/SteamVR/Source/SteamVR/SteamVR.Build.cs
+++ b/Engine/Plugins/Runtime/Steam/SteamVR/Source/SteamVR/SteamVR.Build.cs
@@ -37,11 +37,6 @@ namespace UnrealBuildTool.Rules
                 PrivateDependencyModuleNames.Add("UnrealEd");
             }
 
-			if (Target.Platform == UnrealTargetPlatform.Win64 || Target.IsInPlatformGroup(UnrealPlatformGroup.Linux))
-			{
-				AddEngineThirdPartyPrivateStaticDependencies(Target, "NVAftermath");
-			}
-
             if (Target.Platform == UnrealTargetPlatform.Win64)
             {
 				PrivateDependencyModuleNames.AddRange(
diff --git a/Engine/Plugins/Runtime/TextureShare/Source/TextureShareD3D12/Private/CrossGPUHeap/D3D12CrossGPUHeapFence.cpp b/Engine/Plugins/Runtime/TextureShare/Source/TextureShareD3D12/Private/CrossGPUHeap/D3D12CrossGPUHeapFence.cpp
--- a/Engine/Plugins/Runtime/TextureShare/Source/TextureShareD3D12/Private/CrossGPUHeap/D3D12CrossGPUHeapFence.cpp
+++ b/Engine/Plugins/Runtime/TextureShare/Source/TextureShareD3D12/Private/CrossGPUHeap/D3D12CrossGPUHeapFence.cpp
@@ -5,8 +5,7 @@
 #if TEXTURESHARE_CROSSGPUHEAP
 // DX12 Cross GPU heap resource API (experimental)
 #include "TextureShareD3D12Log.h"
-#include "D3D12RHIPrivate.h"
-#include "D3D12Util.h"
+#include "ID3D12DynamicRHI.h"
 
 #if TE_ENABLE_FENCE
 
diff --git a/Engine/Plugins/Runtime/TextureShare/Source/TextureShareD3D12/Private/CrossGPUHeap/D3D12CrossGPUHeapItem.cpp b/Engine/Plugins/Runtime/TextureShare/Source/TextureShareD3D12/Private/CrossGPUHeap/D3D12CrossGPUHeapItem.cpp
--- a/Engine/Plugins/Runtime/TextureShare/Source/TextureShareD3D12/Private/CrossGPUHeap/D3D12CrossGPUHeapItem.cpp
+++ b/Engine/Plugins/Runtime/TextureShare/Source/TextureShareD3D12/Private/CrossGPUHeap/D3D12CrossGPUHeapItem.cpp
@@ -5,8 +5,7 @@
 #if TEXTURESHARE_CROSSGPUHEAP
 // DX12 Cross GPU heap resource API (experimental)
 #include "TextureShareD3D12Log.h"
-#include "D3D12RHIPrivate.h"
-#include "D3D12Util.h"
+#include "ID3D12DynamicRHI.h"
 
 // macro to deal with COM calls inside a function that returns `{}` on error
 #define CHECK_HR_DEFAULT(COM_call)\
