Oculus OpenXR: workaround Oculus compositor bug by changing layerFlags

Similar workaround is applied if the Oculus plugin is loaded side-by-side, see FOculusOpenXRHMD::OnEndProjectionLayer

diff --git a/Engine/Plugins/Runtime/OpenXR/Source/OpenXRHMD/Private/OpenXRHMD.cpp b/Engine/Plugins/Runtime/OpenXR/Source/OpenXRHMD/Private/OpenXRHMD.cpp
--- a/Engine/Plugins/Runtime/OpenXR/Source/OpenXRHMD/Private/OpenXRHMD.cpp
+++ b/Engine/Plugins/Runtime/OpenXR/Source/OpenXRHMD/Private/OpenXRHMD.cpp
@@ -2300,6 +2300,11 @@ void FOpenXRHMD::OnFinishRendering_RHIThread()
 			{
 				Layer.next = Module->OnEndProjectionLayer(Session, 0, Layer.next, Layer.layerFlags);
 			}
+
+#if PLATFORM_ANDROID
+			// @todo: temporary workaround for Quest compositor issue, see UE-145546
+			Layer.layerFlags |= XR_COMPOSITION_LAYER_UNPREMULTIPLIED_ALPHA_BIT;
+#endif
 		}
 
 		for (const XrCompositionLayerQuad& Quad : PipelinedLayerStateRHI.QuadLayers)
