Validate Oculus requirements on non-Oculus mobile devices

Because we've made our OpenXR code more and more generic (and not tied to Oculus setups), we can actually enter some of the OpenXR code on non-Oculus setups.
This leads to a failure on some devices w/o proper newer Android facilities (it should actually fail 'gracefully' further on down the line on a phone).
We can work around this for now by using GameActivity APIs to check for Oculus presence.

diff --git a/Engine/Plugins/Runtime/OpenXR/Source/OpenXRHMD/Private/OpenXRHMDModule.cpp b/Engine/Plugins/Runtime/OpenXR/Source/OpenXRHMD/Private/OpenXRHMDModule.cpp
--- a/Engine/Plugins/Runtime/OpenXR/Source/OpenXRHMD/Private/OpenXRHMDModule.cpp
+++ b/Engine/Plugins/Runtime/OpenXR/Source/OpenXRHMD/Private/OpenXRHMDModule.cpp
@@ -12,6 +12,8 @@
 #if PLATFORM_ANDROID
 #include <android_native_app_glue.h>
 extern struct android_app* GNativeAndroidApp;
+
+extern bool AndroidThunkCpp_IsOculusMobileApplication();
 #endif
 
 static TAutoConsoleVariable<int32> CVarEnableOpenXRValidationLayer(
@@ -467,6 +469,19 @@ bool FOpenXRHMDModule::InitInstance()
 	// This should only ever be called if we don't already have an instance.
 	check(!Instance);
 
+#if PLATFORM_ANDROID
+	// TODO: Allow OpenXR on non-Oculus Android platforms
+	if (AndroidThunkCpp_IsOculusMobileApplication())
+	{
+		UE_LOG(LogHMD, Log, TEXT("OpenXRHMDModule: App is packaged for Oculus Mobile"));
+	}
+	else
+	{
+		UE_LOG(LogHMD, Log, TEXT("OpenXRHMDModule: App is not packaged for Oculus Mobile"));
+		return false;
+	}
+#endif
+
 	// Get all extension plugins
 	TSet<const ANSICHAR*, AnsiKeyFunc> ExtensionSet;
 	TArray<IOpenXRExtensionPlugin*> ExtModules = IModularFeatures::Get().GetModularFeatureImplementations<IOpenXRExtensionPlugin>(IOpenXRExtensionPlugin::GetModularFeatureName());
