Tracked Geometry 64 bit FVector fixes.
-AR tracking systems deliver 32 bit float data, so we should typically work with that in the lower level code because the data is often transformed into a new space when we copy it into UE so we can iterate the vertices only once if we do the FVector3f to FVector conversion at the same time as we convert for space and worldscale.
-This code is currently only used by the MicrosoftOpenXR plugin, distributed by microsoft.  But all known platforms at this time use float32s for this kind of data.

diff --git a/Engine/Plugins/Runtime/OpenXR/Source/OpenXRAR/Public/IOpenXRARTrackedGeometryHolder.h b/Engine/Plugins/Runtime/OpenXR/Source/OpenXRAR/Public/IOpenXRARTrackedGeometryHolder.h
--- a/Engine/Plugins/Runtime/OpenXR/Source/OpenXRAR/Public/IOpenXRARTrackedGeometryHolder.h
+++ b/Engine/Plugins/Runtime/OpenXR/Source/OpenXRAR/Public/IOpenXRARTrackedGeometryHolder.h
@@ -97,7 +97,7 @@ struct FOpenXRMeshUpdate : public FOpenXRARTrackedGeometryData
 {
 	EARObjectClassification Type = EARObjectClassification::NotApplicable;
 	EARSpatialMeshUsageFlags SpatialMeshUsageFlags = EARSpatialMeshUsageFlags::NotApplicable;
-	TArray<FVector> Vertices;
+	TArray<FVector3f> Vertices;
 	TArray<MRMESH_INDEX_TYPE> Indices;
 
 	FOpenXRMeshUpdate() :
diff --git a/Engine/Plugins/Runtime/OpenXR/Source/OpenXRHMD/Public/OpenXRCore.h b/Engine/Plugins/Runtime/OpenXR/Source/OpenXRHMD/Public/OpenXRCore.h
--- a/Engine/Plugins/Runtime/OpenXR/Source/OpenXRHMD/Public/OpenXRCore.h
+++ b/Engine/Plugins/Runtime/OpenXR/Source/OpenXRHMD/Public/OpenXRCore.h
@@ -51,6 +51,11 @@ FORCEINLINE FVector ToFVector(XrVector3f Vector, float Scale = 1.0f)
 	return FVector(-Vector.z * Scale, Vector.x * Scale, Vector.y * Scale);
 }
 
+FORCEINLINE FVector3f ToFVector3f(XrVector3f Vector, float Scale = 1.0f)
+{
+	return FVector3f(-Vector.z * Scale, Vector.x * Scale, Vector.y * Scale);
+}
+
 FORCEINLINE XrVector3f ToXrVector(FVector Vector, float Scale = 1.0f)
 {
 	if (Vector.IsZero())
