Fix include
#rnx

diff --git a/Engine/Plugins/Runtime/OpenXR/Source/OpenXRHMD/Private/OpenXRHMD.cpp b/Engine/Plugins/Runtime/OpenXR/Source/OpenXRHMD/Private/OpenXRHMD.cpp
--- a/Engine/Plugins/Runtime/OpenXR/Source/OpenXRHMD/Private/OpenXRHMD.cpp
+++ b/Engine/Plugins/Runtime/OpenXR/Source/OpenXRHMD/Private/OpenXRHMD.cpp
@@ -8,7 +8,7 @@
 #include "IOpenXRExtensionPlugin.h"
 
 #include "Misc/App.h"
-#include "Misc/ConfigCache.h"
+#include "Misc/ConfigCacheIni.h"
 #include "Misc/Parse.h"
 #include "Modules/ModuleManager.h"
 #include "EngineGlobals.h"
