VR PIE: Make sure SetIndependentViewportSize is called on GameThread

We have to do a couple things in order to schedule our call to SetIndependentViewportSize on the GameThread

* Defer initial EnableStereo call until VR PIE startup while using the Editor (it used to call EnableStereo at editor startup??)
* VR PIE correctly resizes viewport since the window size and viewport size are decoupled
* Call SetIndependentViewportSize in EnableStereo(true), which happens at VR PIE kickoff time in the GameThread

diff --git a/Engine/Plugins/Runtime/OpenXR/Source/OpenXRHMD/Private/OpenXRHMD.cpp b/Engine/Plugins/Runtime/OpenXR/Source/OpenXRHMD/Private/OpenXRHMD.cpp
--- a/Engine/Plugins/Runtime/OpenXR/Source/OpenXRHMD/Private/OpenXRHMD.cpp
+++ b/Engine/Plugins/Runtime/OpenXR/Source/OpenXRHMD/Private/OpenXRHMD.cpp
@@ -548,6 +548,26 @@ bool FOpenXRHMD::EnableStereo(bool stereo)
 			FApp::SetUseVRFocus(true);
 			FApp::SetHasVRFocus(true);
 
+#if WITH_EDITOR
+			if (GIsEditor)
+			{
+				if (FSceneViewport* SceneVP = FindSceneViewport())
+				{
+					TSharedPtr<SWindow> Window = SceneVP->FindWindow();
+					if (Window.IsValid())
+					{
+						uint32 SizeX = 0;
+						uint32 SizeY = 0;
+						CalculateRenderTargetSize(*SceneVP, SizeX, SizeY);
+
+						// Window continues to be processed when PIE spectator window is minimized
+						Window->SetIndependentViewportSize(FVector2D(SizeX, SizeY));
+					}
+				}
+			}
+
+#endif // WITH_EDITOR
+
 			return true;
 		}
 		return false;
@@ -1670,23 +1690,6 @@ bool FOpenXRHMD::AllocateRenderTargetTexture(uint32 Index, uint32 SizeX, uint32
 		{
 			return false;
 		}
-
-		// image will be acquired next time we begin the rendering
-
-#if WITH_EDITOR
-		if (GIsEditor)
-		{
-			if (FSceneViewport* SceneVP = FindSceneViewport())
-			{
-				TSharedPtr<SWindow> Window = SceneVP->FindWindow();
-				if (Window.IsValid())
-				{
-					// Window continues to be processed when PIE spectator window is minimized
-					Window->SetIndependentViewportSize(FVector2D(SizeX, SizeY));
-				}
-			}
-		}
-#endif
 	}
 
 	// Grab the presentation texture out of the swapchain.
diff --git a/Engine/Source/Runtime/Engine/Private/Slate/SceneViewport.cpp b/Engine/Source/Runtime/Engine/Private/Slate/SceneViewport.cpp
--- a/Engine/Source/Runtime/Engine/Private/Slate/SceneViewport.cpp
+++ b/Engine/Source/Runtime/Engine/Private/Slate/SceneViewport.cpp
@@ -367,7 +367,12 @@ void FSceneViewport::OnDrawViewport( const FGeometry& AllottedGeometry, const FS
 			{
 				//@HACK VREDITOR
 				//check(Window.IsValid());
-				if ( Window->IsViewportSizeDrivenByWindow() )
+
+				// In the stereo case, the HMD display size drives the base RT size, separate from the PIE mirror window
+				const bool bResizeTargetValid = Window->IsViewportSizeDrivenByWindow() || 
+					(IsPlayInEditorViewport() && IsStereoRenderingAllowed());
+
+				if (bResizeTargetValid)
 				{
 					if (ViewportWidget.Pin()->ShouldRenderDirectly())
 					{
diff --git a/Engine/Source/Runtime/Engine/Private/UnrealEngine.cpp b/Engine/Source/Runtime/Engine/Private/UnrealEngine.cpp
--- a/Engine/Source/Runtime/Engine/Private/UnrealEngine.cpp
+++ b/Engine/Source/Runtime/Engine/Private/UnrealEngine.cpp
@@ -3661,7 +3661,9 @@ bool UEngine::InitializeHMDDevice()
 			if (XRSystem.IsValid() && !FParse::Param(FCommandLine::Get(), TEXT("noxrstereo")))
 			{
 				StereoRenderingDevice = XRSystem->GetStereoRenderingDevice();
-				const bool bShouldStartInVR = StereoRenderingDevice.IsValid() && (FParse::Param(FCommandLine::Get(), TEXT("vr")) || GetDefault<UGeneralProjectSettings>()->bStartInVR);
+				const bool bShouldStartInVR = StereoRenderingDevice.IsValid() && 
+											  !GIsEditor &&
+											  (FParse::Param(FCommandLine::Get(), TEXT("vr")) || GetDefault<UGeneralProjectSettings>()->bStartInVR);
 				if (bShouldStartInVR)
 				{
 					StereoRenderingDevice->EnableStereo(true);
