bIsStandaloneStereoOnlyDevice thread safety ensures
-SceneTextures.cpp was using module lookup functions in the render thread that are only safe to use without locking on the game thread.  However we don't want to lock these every frame.  So I have added an accessor for this information to IStereoRendering.h and cached the values in the oculus HMD implementation that uses it but did not already have the value cached in the hmd.  This cached value can be safely read on any thread.

diff --git a/Engine/Plugins/Runtime/Oculus/OculusVR/Source/OculusHMD/Private/OculusHMD.cpp b/Engine/Plugins/Runtime/Oculus/OculusVR/Source/OculusHMD/Private/OculusHMD.cpp
--- a/Engine/Plugins/Runtime/Oculus/OculusVR/Source/OculusHMD/Private/OculusHMD.cpp
+++ b/Engine/Plugins/Runtime/Oculus/OculusVR/Source/OculusHMD/Private/OculusHMD.cpp
@@ -2036,6 +2036,8 @@ namespace OculusHMD
 		SplashLayerHandle = -1;
 
 		SplashRotation = FRotator();
+
+		bIsStandaloneStereoOnlyDevice = IHeadMountedDisplayModule::IsAvailable() && IHeadMountedDisplayModule::Get().IsStandaloneStereoOnlyDevice();;
 	}
 
 
diff --git a/Engine/Plugins/Runtime/Oculus/OculusVR/Source/OculusHMD/Private/OculusHMD.h b/Engine/Plugins/Runtime/Oculus/OculusVR/Source/OculusHMD/Private/OculusHMD.h
--- a/Engine/Plugins/Runtime/Oculus/OculusVR/Source/OculusHMD/Private/OculusHMD.h
+++ b/Engine/Plugins/Runtime/Oculus/OculusVR/Source/OculusHMD/Private/OculusHMD.h
@@ -173,6 +173,7 @@ public:
 	virtual IStereoLayers* GetStereoLayers() override { return this; }
 	//virtual void UseImplicitHmdPosition(bool bInImplicitHmdPosition) override;
 	//virtual bool GetUseImplicitHmdPosition() override;
+	virtual bool IsStandaloneStereoOnlyDevice() const override { return bIsStandaloneStereoOnlyDevice; }
 
 	// FHeadMountedDisplayBase interface
 	virtual FVector2D GetEyeCenterPoint_RenderThread(int32 ViewIndex) const override;
@@ -420,6 +421,7 @@ protected:
 	TWeakPtr<SWindow> CachedWindow;
 	FVector2D CachedWindowSize;
 	float CachedWorldToMetersScale;
+	bool bIsStandaloneStereoOnlyDevice;
 
 	// Game thread
 	FSettingsPtr Settings;
diff --git a/Engine/Plugins/Runtime/OpenXR/Source/OpenXRHMD/Private/OpenXRHMD.h b/Engine/Plugins/Runtime/OpenXR/Source/OpenXRHMD/Private/OpenXRHMD.h
--- a/Engine/Plugins/Runtime/OpenXR/Source/OpenXRHMD/Private/OpenXRHMD.h
+++ b/Engine/Plugins/Runtime/OpenXR/Source/OpenXRHMD/Private/OpenXRHMD.h
@@ -235,7 +235,8 @@ public:
 	virtual int32 GetDesiredNumberOfViews(bool bStereoRequested) const override;
 	virtual EStereoscopicPass GetViewPassForIndex(bool bStereoRequested, int32 ViewIndex) const override;
 	virtual uint32 GetLODViewIndex() const override;
-	
+	virtual bool IsStandaloneStereoOnlyDevice() const override { return bIsStandaloneStereoOnlyDevice; }
+
 	virtual FMatrix GetStereoProjectionMatrix(const int32 StereoViewIndex) const override;
 	virtual void GetEyeRenderParams_RenderThread(const struct FHeadMountedDisplayPassContext& Context, FVector2D& EyeToSrcUVScaleValue, FVector2D& EyeToSrcUVOffsetValue) const override;
 	virtual IStereoRenderTargetManager* GetRenderTargetManager() override;
@@ -308,8 +309,6 @@ public:
 	OPENXRHMD_API XrTime GetDisplayTime() const;
 	OPENXRHMD_API TArray<IOpenXRExtensionPlugin*>& GetExtensionPlugins() { return ExtensionPlugins; }
 
-	OPENXRHMD_API bool IsStandaloneStereoOnlyDevice() const { return bIsStandaloneStereoOnlyDevice; }
-
 private:
 	bool					bStereoEnabled;
 	TAtomic<bool>			bIsRunning;
diff --git a/Engine/Source/Runtime/Engine/Public/StereoRendering.h b/Engine/Source/Runtime/Engine/Public/StereoRendering.h
--- a/Engine/Source/Runtime/Engine/Public/StereoRendering.h
+++ b/Engine/Source/Runtime/Engine/Public/StereoRendering.h
@@ -116,13 +116,21 @@ public:
 	static bool IsASecondaryView(const FSceneView& View);
 
 	/**
-	 * Static helper. Return the index of the view that is used for selecting LODs
+	 * Return the index of the view that is used for selecting LODs
 	 */
 	virtual uint32 GetLODViewIndex() const
 	{
 		return 0;
 	}
 
+	/**
+	 * True when a device has only stereo display (like a self contained mobile vr headset).
+	 */
+	virtual bool IsStandaloneStereoOnlyDevice() const
+	{ 
+		return false; 
+	}
+
 	/**
 	 * Adjusts the viewport rectangle for stereo, based on which eye pass is being rendered.
 	 */
diff --git a/Engine/Source/Runtime/Renderer/Private/SceneTextures.cpp b/Engine/Source/Runtime/Renderer/Private/SceneTextures.cpp
--- a/Engine/Source/Runtime/Renderer/Private/SceneTextures.cpp
+++ b/Engine/Source/Runtime/Renderer/Private/SceneTextures.cpp
@@ -110,8 +110,7 @@ static EPixelFormat GetGBufferFFormat()
 
 static EPixelFormat GetDefaultMobileSceneColorLowPrecisionFormat()
 {
-	return IHeadMountedDisplayModule::IsAvailable() && IHeadMountedDisplayModule::Get().IsStandaloneStereoOnlyDevice()
-		? PF_R8G8B8A8 : PF_B8G8R8A8;
+	return (GEngine && GEngine->XRSystem.IsValid() && GEngine->StereoRenderingDevice.IsValid() && GEngine->StereoRenderingDevice->IsStandaloneStereoOnlyDevice()) ? PF_R8G8B8A8 : PF_B8G8R8A8;
 }
 
 static EPixelFormat GetMobileSceneColorFormat(const FSceneViewFamily& ViewFamily)
