Fix support for the left texture for the Quad stereo layer (UE-146761).

diff --git a/Engine/Plugins/Runtime/OpenXR/Source/OpenXRHMD/Private/OpenXRHMD.cpp b/Engine/Plugins/Runtime/OpenXR/Source/OpenXRHMD/Private/OpenXRHMD.cpp
--- a/Engine/Plugins/Runtime/OpenXR/Source/OpenXRHMD/Private/OpenXRHMD.cpp
+++ b/Engine/Plugins/Runtime/OpenXR/Source/OpenXRHMD/Private/OpenXRHMD.cpp
@@ -1488,8 +1488,8 @@ void FOpenXRHMD::DestroySession()
 		// to destroy swapchain handles after the session is already destroyed.
 		ForEachLayer([&](uint32 /* unused */, FOpenXRLayer& Layer)
 		{
-			Layer.Swapchain.Reset();
-			Layer.LeftSwapchain.Reset();
+			Layer.RightEye.Swapchain.Reset();
+			Layer.LeftEye.Swapchain.Reset();
 		});
 
 		PipelinedLayerStateRendering.ColorSwapchain.Reset();
@@ -1815,26 +1815,23 @@ void FOpenXRHMD::OnBeginRendering_RenderThread(FRHICommandListImmediate& RHICmdL
 				const ETextureCreateFlags Flags = Layer.Desc.Flags & IStereoLayers::LAYER_FLAG_TEX_CONTINUOUS_UPDATE ?
 					TexCreate_Dynamic | TexCreate_SRGB : TexCreate_SRGB;
 
-				if (Layer.NeedReAllocateTexture())
+				if (Layer.NeedReallocateRightTexture())
 				{
 					FRHITexture2D* Texture = Layer.Desc.Texture->GetTexture2D();
-					Layer.Swapchain = CreateSwapchain(Texture, Flags);
-					Layer.SwapchainSize = Texture->GetSizeXY();
-					Layer.bUpdateTexture = true;
+					Layer.RightEye.SetSwapchain(CreateSwapchain(Texture, Flags), Texture->GetSizeXY());
 				}
 
-				if (Layer.NeedReAllocateLeftTexture())
+				if (Layer.NeedReallocateLeftTexture())
 				{
 					FRHITexture2D* Texture = Layer.Desc.LeftTexture->GetTexture2D();
-					Layer.LeftSwapchain = CreateSwapchain(Texture, Flags);
-					Layer.bUpdateTexture = true;
+					Layer.LeftEye.SetSwapchain(CreateSwapchain(Texture, Flags), Texture->GetSizeXY());
 				}
 			}
 			else
 			{
 				// We retain references in FPipelinedLayerState to avoid premature destruction
-				Layer.Swapchain.Reset();
-				Layer.LeftSwapchain.Reset();
+				Layer.RightEye.Swapchain.Reset();
+				Layer.LeftEye.Swapchain.Reset();
 			}
 		});
 	}
@@ -1856,40 +1853,43 @@ void FOpenXRHMD::OnBeginRendering_RenderThread(FRHICommandListImmediate& RHICmdL
 			XR_COMPOSITION_LAYER_BLEND_TEXTURE_SOURCE_ALPHA_BIT;
 		Quad.space = Layer.Desc.PositionType == ELayerType::FaceLocked ?
 			DeviceSpaces[HMDDeviceId].Space : GetTrackingSpace();
-		Quad.subImage.imageRect = ToXrRect(Layer.GetViewport());
 		Quad.subImage.imageArrayIndex = 0;
 		Quad.pose = ToXrPose(Layer.Desc.Transform * PositionTransform, WorldToMeters);
-		Quad.size = ToXrExtent2D(Layer.GetQuadSize(), WorldToMeters);
 
 		// We need to copy each layer into an OpenXR swapchain so they can be displayed by the compositor
-		if (Layer.Swapchain.IsValid() && Layer.Desc.Texture.IsValid())
+		if (Layer.RightEye.Swapchain.IsValid() && Layer.Desc.Texture.IsValid())
 		{
-			if (Layer.bUpdateTexture && bIsRunning)
+			if (Layer.RightEye.bUpdateTexture && bIsRunning)
 			{
 				FRHITexture2D* SrcTexture = Layer.Desc.Texture->GetTexture2D();
-				FIntRect DstRect(FIntPoint(0, 0), Layer.SwapchainSize.IntPoint());
-				CopyTexture_RenderThread(RHICmdList, SrcTexture, FIntRect(), Layer.Swapchain, DstRect, false, bNoAlpha);
+				FIntRect DstRect(FIntPoint(0, 0), Layer.RightEye.SwapchainSize.IntPoint());
+				CopyTexture_RenderThread(RHICmdList, SrcTexture, FIntRect(), Layer.RightEye.Swapchain, DstRect, false, bNoAlpha);
 			}
 
 			Quad.eyeVisibility = bIsStereo ? XR_EYE_VISIBILITY_RIGHT : XR_EYE_VISIBILITY_BOTH;
-			Quad.subImage.swapchain = static_cast<FOpenXRSwapchain*>(Layer.Swapchain.Get())->GetHandle();
+
+			Quad.subImage.imageRect = ToXrRect(Layer.GetRightViewportSize());
+			Quad.subImage.swapchain = static_cast<FOpenXRSwapchain*>(Layer.RightEye.Swapchain.Get())->GetHandle();
+			Quad.size = ToXrExtent2D(Layer.GetRightQuadSize(), WorldToMeters);
 			PipelinedLayerStateRendering.QuadLayers.Add(Quad);
-			PipelinedLayerStateRendering.QuadSwapchains.Add(Layer.Swapchain);
+			PipelinedLayerStateRendering.QuadSwapchains.Add(Layer.RightEye.Swapchain);
 		}
 
-		if (Layer.LeftSwapchain.IsValid() && Layer.Desc.LeftTexture.IsValid())
+		if (Layer.LeftEye.Swapchain.IsValid() && Layer.Desc.LeftTexture.IsValid())
 		{
-			if (Layer.bUpdateTexture && bIsRunning)
+			if (Layer.LeftEye.bUpdateTexture && bIsRunning)
 			{
 				FRHITexture2D* SrcTexture = Layer.Desc.LeftTexture->GetTexture2D();
-				FIntRect DstRect(FIntPoint(0, 0), Layer.SwapchainSize.IntPoint());
-				CopyTexture_RenderThread(RHICmdList, SrcTexture, FIntRect(), Layer.LeftSwapchain, DstRect, false, bNoAlpha);
+				FIntRect DstRect(FIntPoint(0, 0), Layer.LeftEye.SwapchainSize.IntPoint());
+				CopyTexture_RenderThread(RHICmdList, SrcTexture, FIntRect(), Layer.LeftEye.Swapchain, DstRect, false, bNoAlpha);
 			}
 
 			Quad.eyeVisibility = XR_EYE_VISIBILITY_LEFT;
-			Quad.subImage.swapchain = static_cast<FOpenXRSwapchain*>(Layer.LeftSwapchain.Get())->GetHandle();
+			Quad.subImage.imageRect = ToXrRect(Layer.GetLeftViewportSize());
+			Quad.subImage.swapchain = static_cast<FOpenXRSwapchain*>(Layer.LeftEye.Swapchain.Get())->GetHandle();
+			Quad.size = ToXrExtent2D(Layer.GetLeftQuadSize(), WorldToMeters);
 			PipelinedLayerStateRendering.QuadLayers.Add(Quad);
-			PipelinedLayerStateRendering.QuadSwapchains.Add(Layer.LeftSwapchain);
+			PipelinedLayerStateRendering.QuadSwapchains.Add(Layer.LeftEye.Swapchain);
 		}
 	}
 
@@ -1911,7 +1911,9 @@ void FOpenXRHMD::OnBeginRendering_RenderThread(FRHICommandListImmediate& RHICmdL
 		// Reset the update flag on all layers
 		ForEachLayer([&](uint32 /* unused */, FOpenXRLayer& Layer)
 		{
-			Layer.bUpdateTexture = Layer.Desc.Flags & IStereoLayers::LAYER_FLAG_TEX_CONTINUOUS_UPDATE;
+			const bool bUpdateTexture = Layer.Desc.Flags & IStereoLayers::LAYER_FLAG_TEX_CONTINUOUS_UPDATE;
+			Layer.RightEye.bUpdateTexture = bUpdateTexture;
+			Layer.LeftEye.bUpdateTexture = bUpdateTexture;
 		}, false);
 
 		FXRSwapChainPtr ColorSwapchain = PipelinedLayerStateRendering.ColorSwapchain;
diff --git a/Engine/Plugins/Runtime/OpenXR/Source/OpenXRHMD/Private/OpenXRHMD_Layer.cpp b/Engine/Plugins/Runtime/OpenXR/Source/OpenXRHMD/Private/OpenXRHMD_Layer.cpp
--- a/Engine/Plugins/Runtime/OpenXR/Source/OpenXRHMD/Private/OpenXRHMD_Layer.cpp
+++ b/Engine/Plugins/Runtime/OpenXR/Source/OpenXRHMD/Private/OpenXRHMD_Layer.cpp
@@ -4,7 +4,7 @@
 #include "OpenXRCore.h"
 #include "OpenXRPlatformRHI.h"
 
-bool FOpenXRLayer::NeedReAllocateTexture()
+bool FOpenXRLayer::NeedReallocateRightTexture()
 {
 	if (!Desc.Texture.IsValid())
 	{
@@ -17,15 +17,15 @@ bool FOpenXRLayer::NeedReAllocateTexture()
 		return false;
 	}
 
-	if (!Swapchain.IsValid())
+	if (!RightEye.Swapchain.IsValid())
 	{
 		return true;
 	}
 
-	return SwapchainSize != Texture->GetSizeXY();
+	return RightEye.SwapchainSize != Texture->GetSizeXY();
 }
 
-bool FOpenXRLayer::NeedReAllocateLeftTexture()
+bool FOpenXRLayer::NeedReallocateLeftTexture()
 {
 	if (!Desc.LeftTexture.IsValid())
 	{
@@ -38,25 +38,41 @@ bool FOpenXRLayer::NeedReAllocateLeftTexture()
 		return false;
 	}
 
-	if (!Swapchain.IsValid())
+	if (!LeftEye.Swapchain.IsValid())
 	{
 		return true;
 	}
 
-	return SwapchainSize != Texture->GetSizeXY();
+	return LeftEye.SwapchainSize != Texture->GetSizeXY();
 }
 
-FIntRect FOpenXRLayer::GetViewport() const
+FIntRect FOpenXRLayer::GetRightViewportSize() const
 {
-	FBox2D Viewport(SwapchainSize * Desc.UVRect.Min, SwapchainSize * Desc.UVRect.Max);
+	FBox2D Viewport(RightEye.SwapchainSize * Desc.UVRect.Min, RightEye.SwapchainSize * Desc.UVRect.Max);
 	return FIntRect(Viewport.Min.IntPoint(), Viewport.Max.IntPoint());
 }
 
-FVector2D FOpenXRLayer::GetQuadSize() const
+FIntRect FOpenXRLayer::GetLeftViewportSize() const
+{
+	FBox2D Viewport(LeftEye.SwapchainSize * Desc.UVRect.Min, LeftEye.SwapchainSize * Desc.UVRect.Max);
+	return FIntRect(Viewport.Min.IntPoint(), Viewport.Max.IntPoint());
+}
+
+FVector2D FOpenXRLayer::GetRightQuadSize() const
+{
+	if (Desc.Flags & IStereoLayers::LAYER_FLAG_QUAD_PRESERVE_TEX_RATIO)
+	{
+		float AspectRatio = RightEye.SwapchainSize.Y / RightEye.SwapchainSize.X;
+		return FVector2D(Desc.QuadSize.X, Desc.QuadSize.X * AspectRatio);
+	}
+	return Desc.QuadSize;
+}
+
+FVector2D FOpenXRLayer::GetLeftQuadSize() const
 {
 	if (Desc.Flags & IStereoLayers::LAYER_FLAG_QUAD_PRESERVE_TEX_RATIO)
 	{
-		float AspectRatio = SwapchainSize.Y / SwapchainSize.X;
+		float AspectRatio = LeftEye.SwapchainSize.Y / LeftEye.SwapchainSize.X;
 		return FVector2D(Desc.QuadSize.X, Desc.QuadSize.X * AspectRatio);
 	}
 	return Desc.QuadSize;
@@ -80,8 +96,9 @@ void MarkLayerTextureForUpdate(FOpenXRLayer& Layer)
 	// If the swapchain is static we need to re-allocate it before it can be updated
 	if (!(Layer.Desc.Flags & IStereoLayers::LAYER_FLAG_TEX_CONTINUOUS_UPDATE))
 	{
-		Layer.Swapchain.Reset();
-		Layer.LeftSwapchain.Reset();
+		Layer.RightEye.Swapchain.Reset();
+		Layer.LeftEye.Swapchain.Reset();
 	}
-	Layer.bUpdateTexture = true;
+	Layer.RightEye.bUpdateTexture = true;
+	Layer.LeftEye.bUpdateTexture = true;
 }
diff --git a/Engine/Plugins/Runtime/OpenXR/Source/OpenXRHMD/Private/OpenXRHMD_Layer.h b/Engine/Plugins/Runtime/OpenXR/Source/OpenXRHMD/Private/OpenXRHMD_Layer.h
--- a/Engine/Plugins/Runtime/OpenXR/Source/OpenXRHMD/Private/OpenXRHMD_Layer.h
+++ b/Engine/Plugins/Runtime/OpenXR/Source/OpenXRHMD/Private/OpenXRHMD_Layer.h
@@ -8,27 +8,43 @@
 
 struct FOpenXRLayer
 {
+	struct FPerEyeTextureData
+	{
+		FXRSwapChainPtr			Swapchain = nullptr;
+		FVector2D				SwapchainSize{};
+		bool					bUpdateTexture = false;
+
+		void					SetSwapchain(FXRSwapChainPtr InSwapchain, const FVector2D& InSwapchainSize)
+		{
+			Swapchain = InSwapchain;
+			SwapchainSize = InSwapchainSize;
+			bUpdateTexture = true;
+		}
+	};
+
 	IStereoLayers::FLayerDesc	Desc;
-	FXRSwapChainPtr				Swapchain;
-	FXRSwapChainPtr				LeftSwapchain;
-	FVector2D					SwapchainSize;
-	bool						bUpdateTexture;
+
+	/** Texture tracking data for the right eye.*/
+	FPerEyeTextureData			RightEye;
+
+	/** Texture tracking data for the left eye, may not be present.*/
+	FPerEyeTextureData			LeftEye;
 
 	FOpenXRLayer(const IStereoLayers::FLayerDesc& InLayerDesc)
 		: Desc(InLayerDesc)
-		, Swapchain()
-		, LeftSwapchain()
-		, bUpdateTexture(false)
 	{ }
 
 	void SetLayerId(uint32 InId) { Desc.SetLayerId(InId); }
 	uint32 GetLayerId() const { return Desc.GetLayerId(); }
 
-	bool NeedReAllocateTexture();
-	bool NeedReAllocateLeftTexture();
+	bool NeedReallocateRightTexture();
+	bool NeedReallocateLeftTexture();
+
+	FIntRect GetRightViewportSize() const;
+	FVector2D GetRightQuadSize() const;
 
-	FIntRect GetViewport() const;
-	FVector2D GetQuadSize() const;
+	FIntRect GetLeftViewportSize() const;
+	FVector2D GetLeftQuadSize() const;
 };
 
 bool GetLayerDescMember(const FOpenXRLayer& Layer, IStereoLayers::FLayerDesc& OutLayerDesc);
