OpenXR xrGetReferenceSpaceBoundsRect make ensure on failure into warning.
-xrGetReferenceSpaceBoundsRect fails for non-stage reference spaces.  This failure is not super-critical, and can be induced by asking for the wrong space type in blueprint, so it should be a warning rather than an ensure.

diff --git a/Engine/Plugins/Runtime/OpenXR/Source/OpenXRHMD/Private/OpenXRHMD.cpp b/Engine/Plugins/Runtime/OpenXR/Source/OpenXRHMD/Private/OpenXRHMD.cpp
--- a/Engine/Plugins/Runtime/OpenXR/Source/OpenXRHMD/Private/OpenXRHMD.cpp
+++ b/Engine/Plugins/Runtime/OpenXR/Source/OpenXRHMD/Private/OpenXRHMD.cpp
@@ -288,7 +288,13 @@ FVector2D FOpenXRHMD::GetPlayAreaBounds(EHMDTrackingOrigin::Type Origin) const
 		break;
 	}
 	XrExtent2Df Bounds;
-	XR_ENSURE(xrGetReferenceSpaceBoundsRect(Session, Space, &Bounds));
+	const XrResult Result = xrGetReferenceSpaceBoundsRect(Session, Space, &Bounds);
+	if (Result != XR_SUCCESS)
+	{
+		UE_LOG(LogHMD, Warning, TEXT("GetPlayAreaBounds xrGetReferenceSpaceBoundsRect with reference space %s failed with result %s. Returning zero vector."), OpenXRReferenceSpaceTypeToString(Space), OpenXRResultToString(Result));
+		return FVector2D::ZeroVector;
+	}
+	
 	Swap(Bounds.width, Bounds.height); // Convert to UE coordinate system
 	return ToFVector2D(Bounds, WorldToMetersScale);
 }
diff --git a/Engine/Plugins/Runtime/OpenXR/Source/OpenXRHMD/Public/OpenXRCore.h b/Engine/Plugins/Runtime/OpenXR/Source/OpenXRHMD/Public/OpenXRCore.h
--- a/Engine/Plugins/Runtime/OpenXR/Source/OpenXRHMD/Public/OpenXRCore.h
+++ b/Engine/Plugins/Runtime/OpenXR/Source/OpenXRHMD/Public/OpenXRCore.h
@@ -27,6 +27,16 @@ constexpr const TCHAR* OpenXRSessionStateToString(XrSessionState e)
 	}
 }
 
+#define XR_REFERENCE_SPACE_TYPE_STR(name, val) case name: return TEXT(#name);
+constexpr const TCHAR* OpenXRReferenceSpaceTypeToString(XrReferenceSpaceType e)
+{
+	switch (e)
+	{
+		XR_LIST_ENUM_XrReferenceSpaceType(XR_REFERENCE_SPACE_TYPE_STR);
+	default: return TEXT("Unknown");
+	}
+}
+
 #if DO_CHECK
 #define XR_ENSURE(x) [] (XrResult Result) \
 	{ \
