UE-148741 OpenXR xrGetInstanceProperties string safety
-bulletproof the runtime name buffer by null terminator initializing it and ensuring it has a null terminator after calling it.

diff --git a/Engine/Plugins/Runtime/OpenXR/Source/OpenXRHMD/Private/OpenXRHMD.cpp b/Engine/Plugins/Runtime/OpenXR/Source/OpenXRHMD/Private/OpenXRHMD.cpp
--- a/Engine/Plugins/Runtime/OpenXR/Source/OpenXRHMD/Private/OpenXRHMD.cpp
+++ b/Engine/Plugins/Runtime/OpenXR/Source/OpenXRHMD/Private/OpenXRHMD.cpp
@@ -914,6 +914,7 @@ FOpenXRHMD::FOpenXRHMD(const FAutoRegister& AutoRegister, XrInstance InInstance,
 {
 	InstanceProperties = { XR_TYPE_INSTANCE_PROPERTIES, nullptr };
 	XR_ENSURE(xrGetInstanceProperties(Instance, &InstanceProperties));
+	InstanceProperties.runtimeName[XR_MAX_RUNTIME_NAME_SIZE - 1] = 0; // Ensure the name is null terminated.
 
 	bDepthExtensionSupported = IsExtensionEnabled(XR_KHR_COMPOSITION_LAYER_DEPTH_EXTENSION_NAME) && CheckPlatformDepthExtensionSupport(InstanceProperties);
 
diff --git a/Engine/Plugins/Runtime/OpenXR/Source/OpenXRHMD/Private/OpenXRHMDModule.cpp b/Engine/Plugins/Runtime/OpenXR/Source/OpenXRHMD/Private/OpenXRHMDModule.cpp
--- a/Engine/Plugins/Runtime/OpenXR/Source/OpenXRHMD/Private/OpenXRHMDModule.cpp
+++ b/Engine/Plugins/Runtime/OpenXR/Source/OpenXRHMD/Private/OpenXRHMDModule.cpp
@@ -676,10 +676,9 @@ bool FOpenXRHMDModule::InitInstance()
 		return false;
 	}
 
-	XrInstanceProperties InstanceProps;
-	InstanceProps.type = XR_TYPE_INSTANCE_PROPERTIES;
-	InstanceProps.next = nullptr;
+	XrInstanceProperties InstanceProps = { XR_TYPE_INSTANCE_PROPERTIES, nullptr };
 	XR_ENSURE(xrGetInstanceProperties(Instance, &InstanceProps));
+	InstanceProps.runtimeName[XR_MAX_RUNTIME_NAME_SIZE - 1] = 0; // Ensure the name is null terminated.
 	UE_LOG(LogHMD, Log, TEXT("Initialized OpenXR on %S runtime version %d.%d.%d"), InstanceProps.runtimeName, XR_VERSION_MAJOR(InstanceProps.runtimeVersion), XR_VERSION_MINOR(InstanceProps.runtimeVersion), XR_VERSION_PATCH(InstanceProps.runtimeVersion));
 
 	for (IOpenXRExtensionPlugin* Module : ExtensionPlugins)
