OpenXR: Handoff ERHIAccess to CopyTexture_RenderThread

The VulkanRHI doesn't like certain combinations of ERHIAccess flags, so we make sure to split them across different usages of OpenXRHMD::CopyTexture_RenderThread (spectator screen vs swapchain layers).

diff --git a/Engine/Plugins/Runtime/OpenXR/Source/OpenXRHMD/Private/OpenXRHMD.cpp b/Engine/Plugins/Runtime/OpenXR/Source/OpenXRHMD/Private/OpenXRHMD.cpp
--- a/Engine/Plugins/Runtime/OpenXR/Source/OpenXRHMD/Private/OpenXRHMD.cpp
+++ b/Engine/Plugins/Runtime/OpenXR/Source/OpenXRHMD/Private/OpenXRHMD.cpp
@@ -2401,7 +2401,8 @@ FIntRect FOpenXRHMD::GetFullFlatEyeRect_RenderThread(FTexture2DRHIRef EyeTexture
 	return FIntRect(EyeTexture->GetSizeX() * SrcNormRectMin.X, EyeTexture->GetSizeY() * SrcNormRectMin.Y, EyeTexture->GetSizeX() * SrcNormRectMax.X, EyeTexture->GetSizeY() * SrcNormRectMax.Y);
 }
 
-void FOpenXRHMD::CopyTexture_RenderThread(FRHICommandListImmediate& RHICmdList, FRHITexture2D* SrcTexture, FIntRect SrcRect, FRHITexture2D* DstTexture, FIntRect DstRect, bool bClearBlack, bool bNoAlpha, ERenderTargetActions RTAction) const
+void FOpenXRHMD::CopyTexture_RenderThread(FRHICommandListImmediate& RHICmdList, FRHITexture2D* SrcTexture, FIntRect SrcRect, FRHITexture2D* DstTexture, FIntRect DstRect, 
+											bool bClearBlack, bool bNoAlpha, ERenderTargetActions RTAction, ERHIAccess FinalDstAccess) const
 {
 	check(IsInRenderingThread());
 
@@ -2487,12 +2488,7 @@ void FOpenXRHMD::CopyTexture_RenderThread(FRHICommandListImmediate& RHICmdList,
 	}
 	RHICmdList.EndRenderPass();
 	
-	// TODO
-	// We can't assume the new access state is valid for usage by clients. For the spectator screen, Present is fine.
-	// For layer blits, we'll need to have something that's compatible with swapchain requirements when handed off to
-	// compositor (RTV?). There's a pending refactor from Jules that will make it easier for caller to pass in
-	// desired ERHIAccess for destination texture.
-	RHICmdList.Transition(FRHITransitionInfo(DstTexture, ERHIAccess::RTV, (ERHIAccess::Present | ERHIAccess::SRVMask)));
+	RHICmdList.Transition(FRHITransitionInfo(DstTexture, ERHIAccess::RTV, FinalDstAccess));
 }
 
 void FOpenXRHMD::CopyTexture_RenderThread(FRHICommandListImmediate& RHICmdList, FRHITexture2D* SrcTexture, FIntRect SrcRect, const FXRSwapChainPtr& DstSwapChain, FIntRect DstRect, bool bClearBlack, bool bNoAlpha) const
@@ -2505,7 +2501,7 @@ void FOpenXRHMD::CopyTexture_RenderThread(FRHICommandListImmediate& RHICmdList,
 
 	// Now that we've enqueued the swapchain wait we can add the commands to do the actual texture copy
 	FRHITexture2D* const DstTexture = DstSwapChain->GetTexture2DArray() ? DstSwapChain->GetTexture2DArray() : DstSwapChain->GetTexture2D();
-	CopyTexture_RenderThread(RHICmdList, SrcTexture, SrcRect, DstTexture, DstRect, bClearBlack, bNoAlpha, ERenderTargetActions::Clear_Store);
+	CopyTexture_RenderThread(RHICmdList, SrcTexture, SrcRect, DstTexture, DstRect, bClearBlack, bNoAlpha, ERenderTargetActions::Clear_Store, ERHIAccess::SRVMask);
 
 	// Enqueue a command to release the image after the copy is done
 	RHICmdList.EnqueueLambda([DstSwapChain](FRHICommandListImmediate& InRHICmdList)
@@ -2517,7 +2513,7 @@ void FOpenXRHMD::CopyTexture_RenderThread(FRHICommandListImmediate& RHICmdList,
 void FOpenXRHMD::CopyTexture_RenderThread(FRHICommandListImmediate& RHICmdList, FRHITexture2D* SrcTexture, FIntRect SrcRect, FRHITexture2D* DstTexture, FIntRect DstRect, bool bClearBlack, bool bNoAlpha) const
 {
 	// FIXME: This should probably use the Load_Store action since the spectator controller does multiple overlaying copies.
-	CopyTexture_RenderThread(RHICmdList, SrcTexture, SrcRect, DstTexture, DstRect, bClearBlack, bNoAlpha, ERenderTargetActions::DontLoad_Store);
+	CopyTexture_RenderThread(RHICmdList, SrcTexture, SrcRect, DstTexture, DstRect, bClearBlack, bNoAlpha, ERenderTargetActions::DontLoad_Store, ERHIAccess::Present);
 }
 
 void FOpenXRHMD::RenderTexture_RenderThread(class FRHICommandListImmediate& RHICmdList, class FRHITexture* BackBuffer, class FRHITexture* SrcTexture, FVector2D WindowSize) const
diff --git a/Engine/Plugins/Runtime/OpenXR/Source/OpenXRHMD/Private/OpenXRHMD.h b/Engine/Plugins/Runtime/OpenXR/Source/OpenXRHMD/Private/OpenXRHMD.h
--- a/Engine/Plugins/Runtime/OpenXR/Source/OpenXRHMD/Private/OpenXRHMD.h
+++ b/Engine/Plugins/Runtime/OpenXR/Source/OpenXRHMD/Private/OpenXRHMD.h
@@ -195,7 +195,9 @@ protected:
 	void EnumerateViews(FPipelinedFrameState& PipelineState);
 	void LocateViews(FPipelinedFrameState& PipelinedState, bool ResizeViewsArray = false);
 
-	void CopyTexture_RenderThread(FRHICommandListImmediate& RHICmdList, FRHITexture2D* SrcTexture, FIntRect SrcRect, FRHITexture2D* DstTexture, FIntRect DstRect, bool bClearBlack, bool bNoAlpha, ERenderTargetActions RTAction) const;
+	void CopyTexture_RenderThread(FRHICommandListImmediate& RHICmdList, FRHITexture2D* SrcTexture, FIntRect SrcRect, FRHITexture2D* DstTexture, FIntRect DstRect, 
+								  bool bClearBlack, bool bNoAlpha, ERenderTargetActions RTAction, ERHIAccess FinalDstAccess) const;
+	
 	void CopyTexture_RenderThread(FRHICommandListImmediate& RHICmdList, FRHITexture2D* SrcTexture, FIntRect SrcRect, const FXRSwapChainPtr& DstSwapChain, FIntRect DstRect, bool bClearBlack, bool bNoAlpha) const;
 
 public:
