OpenXRHMD: Fix up rendering thread accesses for onethread

onethread needs a couple corrected accesses that get mixed up with current usage of IsInRenderingThread

diff --git a/Engine/Plugins/Runtime/OpenXR/Source/OpenXRHMD/Private/OpenXRHMD.cpp b/Engine/Plugins/Runtime/OpenXR/Source/OpenXRHMD/Private/OpenXRHMD.cpp
--- a/Engine/Plugins/Runtime/OpenXR/Source/OpenXRHMD/Private/OpenXRHMD.cpp
+++ b/Engine/Plugins/Runtime/OpenXR/Source/OpenXRHMD/Private/OpenXRHMD.cpp
@@ -267,7 +267,7 @@ void FOpenXRHMD::GetMotionControllerData(UObject* WorldContext, const EControlle
 
 float FOpenXRHMD::GetWorldToMetersScale() const
 {
-	return IsInRenderingThread() ? PipelinedFrameStateRendering.WorldToMetersScale : PipelinedFrameStateGame.WorldToMetersScale;
+	return IsInActualRenderingThread() ? PipelinedFrameStateRendering.WorldToMetersScale : PipelinedFrameStateGame.WorldToMetersScale;
 }
 
 FVector2D FOpenXRHMD::GetPlayAreaBounds(EHMDTrackingOrigin::Type Origin) const
@@ -625,6 +625,8 @@ void FOpenXRHMD::AdjustViewRect(int32 ViewIndex, int32& X, int32& Y, uint32& Siz
 
 void FOpenXRHMD::SetFinalViewRect(FRHICommandListImmediate& RHICmdList, const int32 ViewIndex, const FIntRect& FinalViewRect)
 {
+	check(IsInRenderingThread());
+
 	if (ViewIndex == INDEX_NONE || !PipelinedLayerStateRendering.ColorImages.IsValidIndex(ViewIndex))
 	{
 		return;
@@ -1042,7 +1044,7 @@ const FOpenXRHMD::FPipelinedFrameState& FOpenXRHMD::GetPipelinedFrameStateForThr
 	// Relying on implicit selection of the RHI struct is hazardous since the RHI thread isn't always present
 	check(!IsInRHIThread());
 
-	if (IsInRenderingThread())
+	if (IsInActualRenderingThread())
 	{
 		return PipelinedFrameStateRendering;
 	}
@@ -1058,7 +1060,7 @@ FOpenXRHMD::FPipelinedFrameState& FOpenXRHMD::GetPipelinedFrameStateForThread()
 	// Relying on implicit selection of the RHI struct is hazardous since the RHI thread isn't always present
 	check(!IsInRHIThread());
 
-	if (IsInRenderingThread())
+	if (IsInActualRenderingThread())
 	{
 		return PipelinedFrameStateRendering;
 	}
