OpenXR: Support code in external modules participating in binding suggestions for built-in interaction profiles.

diff --git a/Engine/Plugins/Runtime/OpenXR/Source/OpenXRHMD/Public/IOpenXRExtensionPlugin.h b/Engine/Plugins/Runtime/OpenXR/Source/OpenXRHMD/Public/IOpenXRExtensionPlugin.h
--- a/Engine/Plugins/Runtime/OpenXR/Source/OpenXRHMD/Public/IOpenXRExtensionPlugin.h
+++ b/Engine/Plugins/Runtime/OpenXR/Source/OpenXRHMD/Public/IOpenXRExtensionPlugin.h
@@ -115,7 +115,7 @@ public:
 	/**
 	* Register module as an extension on startup.  
 	* It is common to do this in StartupModule of your IModuleInterface class (which may also be the class that implements this interface).
-	* The module's LoadingPhase must be PostInitConfig or earlier because OpenXRHMD will look for these after it is loaded in that phase.
+	* The module's LoadingPhase must be PostConfigInit or earlier because OpenXRHMD will look for these after it is loaded in that phase.
 	*/
 	void RegisterOpenXRExtensionModularFeature()
 	{
@@ -217,21 +217,59 @@ public:
 	}
 
 	/**
-	* Add any actions provided by the plugin to Actions with suggested bindings.
+	* Add any actions provided by the plugin to Actions.
 	* This allows a plugin to 'hard code' an action so that the plugin can use it.
 	*/
+	UE_DEPRECATED(5.1, "Use the overload taking FActionParams/FActionSetParams below instead.")
 	virtual void AddActions(XrInstance Instance, TFunction<XrAction(XrActionType InActionType, const FName& InName, const TArray<XrPath>& InSubactionPaths)> AddAction)
 	{
 	}
 
+	struct FActionSetParams
+	{
+		FName Name;
+		FText LocalizedName;
+		uint32 Priority;
+	};
+	using FCreateActionSetFunc = TFunction< XrActionSet(const FActionSetParams&) >;
+
+	struct FActionParams
+	{
+		/** If Set == XR_NULL_HANDLE, the action is added to the default action set. */
+		XrActionSet Set = XR_NULL_HANDLE;
+
+		XrActionType Type;
+		FName Name;
+		FText LocalizedName;
+		TArray<XrPath> SubactionPaths;
+		TArray<FKey> SuggestedBindings;
+	};
+	using FCreateActionFunc = TFunction< XrAction(const FActionParams&) >;
+
+	/**
+	* Create and register actions and/or action sets belonging to the plugin.
+	* This allows a plugin to create action sets to be attached to the session,
+	* and actions which participate in interaction profile binding suggestions.
+	*/
+	virtual void AddActions(XrInstance InInstance, FCreateActionSetFunc CreateActionSet, FCreateActionFunc CreateAction)
+	{
+	}
+
 	/**
 	* Add any action sets provided by the plugin to be attached as active to the session
 	* This allows a plugin to manage a custom actionset that will be active in xrSyncActions
 	*/
+	UE_DEPRECATED(5.1, "Use AddActions to create attached action sets, along with GetActiveActionSetsForSync.")
 	virtual void AddActionSets(TArray<XrActiveActionSet>& OutActionSets)
 	{
 	}
 
+	/**
+	* Specify action sets to be included in XrActionsSyncInfo::activeActionSets.
+	*/
+	virtual void GetActiveActionSetsForSync(TArray<XrActiveActionSet>& OutActiveSets)
+	{
+	}
 
 	/**
 	* Use this callback to handle events that the OpenXR plugin doesn't handle itself
diff --git a/Engine/Plugins/Runtime/OpenXR/Source/OpenXRInput/Private/OpenXRInput.cpp b/Engine/Plugins/Runtime/OpenXR/Source/OpenXRInput/Private/OpenXRInput.cpp
--- a/Engine/Plugins/Runtime/OpenXR/Source/OpenXRInput/Private/OpenXRInput.cpp
+++ b/Engine/Plugins/Runtime/OpenXR/Source/OpenXRInput/Private/OpenXRInput.cpp
@@ -261,7 +261,7 @@ void FOpenXRInputPlugin::FOpenXRInput::BuildActions()
 	SetInfo.type = XR_TYPE_ACTION_SET_CREATE_INFO;
 	SetInfo.next = nullptr;
 	FCStringAnsi::Strcpy(SetInfo.actionSetName, XR_MAX_ACTION_SET_NAME_SIZE, "ue");
-	FCStringAnsi::Strcpy(SetInfo.localizedActionSetName, XR_MAX_ACTION_SET_NAME_SIZE, "Unreal Engine");
+	FCStringAnsi::Strcpy(SetInfo.localizedActionSetName, XR_MAX_LOCALIZED_ACTION_SET_NAME_SIZE, "Unreal Engine");
 	SetInfo.priority = 0;
 	XR_ENSURE(xrCreateActionSet(Instance, &SetInfo, &ActionSet));
 
@@ -348,14 +348,75 @@ void FOpenXRInputPlugin::FOpenXRInput::BuildActions()
 	// Query extension plugins for actions
 	for (IOpenXRExtensionPlugin* Plugin : OpenXRHMD->GetExtensionPlugins())
 	{
+		PRAGMA_DISABLE_DEPRECATION_WARNINGS;
 		Plugin->AddActions(Instance,
 			[this, &ActionSet](XrActionType InActionType, const FName& InName, const TArray<XrPath>& InSubactionPaths)
 			{
+				// TODO?: Log deprecation warning at runtime, since overridden deprecated interface methods don't warn at compile time?
 				FOpenXRAction Action(ActionSet, InActionType, InName, InSubactionPaths);
 				Actions.Add(Action);
 				return Action.Handle;
 			}
-			);
+		);
+		PRAGMA_ENABLE_DEPRECATION_WARNINGS;
+	}
+
+	for (IOpenXRExtensionPlugin* Plugin : OpenXRHMD->GetExtensionPlugins())
+	{
+		IOpenXRExtensionPlugin::FCreateActionSetFunc CreateActionSetFunc =
+			[this, Instance]
+			(const IOpenXRExtensionPlugin::FActionSetParams& Params) -> XrActionSet
+			{
+				XrActionSet OutActionSet = XR_NULL_HANDLE;
+
+				char ActionSetName[NAME_SIZE];
+				Params.Name.GetPlainANSIString(ActionSetName);
+
+				XrActionSetCreateInfo SetInfo;
+				SetInfo.type = XR_TYPE_ACTION_SET_CREATE_INFO;
+				SetInfo.next = nullptr;
+				FCStringAnsi::Strcpy(SetInfo.actionSetName, XR_MAX_ACTION_SET_NAME_SIZE, ActionSetName);
+				SetInfo.priority = Params.Priority;
+
+				if (!Params.LocalizedName.IsEmpty())
+				{
+					const FString& LocalizedNameStr = Params.LocalizedName.ToString();
+					FTCHARToUTF8_Convert::Convert(SetInfo.localizedActionSetName, XR_MAX_LOCALIZED_ACTION_SET_NAME_SIZE, *LocalizedNameStr, LocalizedNameStr.Len() + 1);
+				}
+				else
+				{
+					FCStringAnsi::Strcpy(SetInfo.localizedActionSetName, XR_MAX_LOCALIZED_ACTION_SET_NAME_SIZE, ActionSetName);
+				}
+
+
+				if (XR_ENSURE(xrCreateActionSet(Instance, &SetInfo, &OutActionSet)))
+				{
+					PluginActionSets.Add(OutActionSet);
+					return OutActionSet;
+				}
+				else
+				{
+					return XR_NULL_HANDLE;
+				}
+			};
+
+		IOpenXRExtensionPlugin::FCreateActionFunc CreateActionFunc =
+			[this, Instance, ActionSet, &Profiles]
+			(const IOpenXRExtensionPlugin::FActionParams& Params) -> XrAction
+			{
+				// TODO: Use Params.LocalizedName
+				const XrActionSet UseActionSet = Params.Set == XR_NULL_HANDLE ? ActionSet : Params.Set;
+				FOpenXRAction Action(UseActionSet, Params.Type, Params.Name, Params.SubactionPaths);
+				const XrAction ReturnHandle = Action.Handle;
+				for (const FKey& Key : Params.SuggestedBindings)
+				{
+					SuggestBindingForKey(Instance, Action, Key, Profiles);
+				}
+				Actions.Emplace(MoveTemp(Action));
+				return ReturnHandle;
+			};
+
+		Plugin->AddActions(Instance, CreateActionSetFunc, CreateActionFunc);
 	}
 
 	for (TPair<FString, FInteractionProfile>& Pair : Profiles)
@@ -428,6 +489,7 @@ void FOpenXRInputPlugin::FOpenXRInput::DestroyActions()
 	Controllers.Reset();
 	SubactionPaths.Reset();
 	ActionSets.Reset();
+	PluginActionSets.Reset();
 }
 
 template<typename T>
@@ -438,88 +500,98 @@ int32 FOpenXRInputPlugin::FOpenXRInput::SuggestBindings(XrInstance Instance, FOp
 	// Add suggested bindings for every mapping
 	for (const T& InputKey : Mappings)
 	{
-		// Key names that are parseable into an OpenXR path have exactly 4 tokens
-		TArray<FString> Tokens;
-		if (InputKey.Key.ToString().ParseIntoArray(Tokens, TEXT("_")) != EKeys::NUM_XR_KEY_TOKENS)
+		if (SuggestBindingForKey(Instance, Action, InputKey.Key, Profiles))
 		{
-			continue;
+			++SuggestedBindings;
 		}
+	}
 
-		// Check if we support the profile specified in the key name
-		FInteractionProfile* Profile = Profiles.Find(Tokens[0]);
-		if (Profile)
-		{
-			// Parse the key name into an OpenXR interaction profile path
-			FString Path = "/user/hand/" + Tokens[1].ToLower();
-			XrPath TopLevel = GetPath(Instance, Path);
-
-			// Map this key to the correct subaction for this profile
-			// We'll use this later to trigger the correct key
-			TPair<XrPath, XrPath> Key(Profile->Path, TopLevel);
-			Action.KeyMap.Add(Key, InputKey.Key.GetFName());
-
-			// Add the input we want to query with grip being defined as "squeeze" in OpenXR
-			FString Identifier = Tokens[2].ToLower();
-			if (Identifier == "grip")
-			{
-				Identifier = "squeeze";
-			}
-			Path += "/input/" + Identifier;
+	return SuggestedBindings;
+}
 
-			// Add the data we want to query, we'll skip this for trigger/squeeze "click" actions to allow
-			// certain profiles that don't have "click" data to threshold the "value" data instead
-			FString Component = Tokens[3].ToLower();
-			if (Component == "axis")
-			{
-				Path += "/value";
-			}
-			else if (Component == "click")
-			{
-				if (Tokens[0] == "ValveIndex" && (Identifier == "trackpad" || Identifier == "squeeze"))
-				{
-					continue;
-				}
+bool FOpenXRInputPlugin::FOpenXRInput::SuggestBindingForKey(XrInstance Instance, FOpenXRAction& Action, const FKey& InFKey, TMap<FString, FInteractionProfile>& Profiles)
+{
+	// Key names that are parseable into an OpenXR path have exactly 4 tokens
+	TArray<FString> Tokens;
+	if (InFKey.ToString().ParseIntoArray(Tokens, TEXT("_")) != EKeys::NUM_XR_KEY_TOKENS)
+	{
+		return false;
+	}
 
-				// The OpenXR spec says that .../input/system/click might not be available for application usage
-				if (Identifier == "system")
-				{
-					continue;
-				}
+	// Check if we support the profile specified in the key name
+	FInteractionProfile* Profile = Profiles.Find(Tokens[0]);
+	if (!Profile)
+	{
+		return false;
+	}
 
-				if (Identifier != "trigger" && Identifier != "squeeze")
-				{
-					Path += "/click";
-				}
-			}
-			else if (Component == "touch")
-			{
-				Path += "/touch";
-			}
-			else if (Component == "touchaxis")
-			{
-				Path += "/touchvalue";  // Note: this is not a standard openxr identifier.  It is meant to represent some kind of analog touch sensor.
-			}
-			else if (Component == "up" || Component == "down" || Component == "left" || Component == "right")
-			{
-				if (!bDirectionalBindingSupported)
-				{
-					continue;
-				}
-				Path += "/dpad_" + Component;
-			}
-			else
-			{
-				// Anything we don't need to translate can pass through
-				Path += "/" + Component;
-			}
+	// Parse the key name into an OpenXR interaction profile path
+	FString Path = "/user/hand/" + Tokens[1].ToLower();
+	XrPath TopLevel = GetPath(Instance, Path);
+
+	// Map this key to the correct subaction for this profile
+	// We'll use this later to trigger the correct key
+	TPair<XrPath, XrPath> Key(Profile->Path, TopLevel);
+	Action.KeyMap.Add(Key, InFKey.GetFName());
+
+	// Add the input we want to query with grip being defined as "squeeze" in OpenXR
+	FString Identifier = Tokens[2].ToLower();
+	if (Identifier == "grip")
+	{
+		Identifier = "squeeze";
+	}
+	Path += "/input/" + Identifier;
+
+	// Add the data we want to query, we'll skip this for trigger/squeeze "click" actions to allow
+	// certain profiles that don't have "click" data to threshold the "value" data instead
+	FString Component = Tokens[3].ToLower();
+	if (Component == "axis")
+	{
+		Path += "/value";
+	}
+	else if (Component == "click")
+	{
+		if (Tokens[0] == "ValveIndex" && (Identifier == "trackpad" || Identifier == "squeeze"))
+		{
+			return false;
+		}
 
-			// Add the binding to the profile
-			Profile->Bindings.Add(XrActionSuggestedBinding{ Action.Handle, GetPath(Instance, Path) });
-			SuggestedBindings++;
+		// The OpenXR spec says that .../input/system/click might not be available for application usage
+		if (Identifier == "system")
+		{
+			return false;
 		}
+
+		if (Identifier != "trigger" && Identifier != "squeeze")
+		{
+			Path += "/click";
+		}
+	}
+	else if (Component == "touch")
+	{
+		Path += "/touch";
+	}
+	else if (Component == "touchaxis")
+	{
+		Path += "/touchvalue";  // Note: this is not a standard openxr identifier.  It is meant to represent some kind of analog touch sensor.
+	}
+	else if (Component == "up" || Component == "down" || Component == "left" || Component == "right")
+	{
+		if (!bDirectionalBindingSupported)
+		{
+			return false;
+		}
+		Path += "/dpad_" + Component;
+	}
+	else
+	{
+		// Anything we don't need to translate can pass through
+		Path += "/" + Component;
 	}
 
-	return SuggestedBindings;
+	// Add the binding to the profile
+	Profile->Bindings.Add(XrActionSuggestedBinding{ Action.Handle, GetPath(Instance, Path) });
+	return true;
 }
 
 void FOpenXRInputPlugin::FOpenXRInput::Tick(float DeltaTime)
@@ -540,12 +612,22 @@ void FOpenXRInputPlugin::FOpenXRInput::Tick(float DeltaTime)
 			for (auto && BindActionSet : ActionSets)
 				BindActionSets.Add(BindActionSet.actionSet);
 
-			// Bind plugin action sets
-			PluginActionSets.Empty();
+			// Bind plugin action sets exposed via deprecated method
 			for (IOpenXRExtensionPlugin* Plugin : OpenXRHMD->GetExtensionPlugins())
-				Plugin->AddActionSets(PluginActionSets);
+			{
+				TArray<XrActiveActionSet> PluginAttachSets_Deprecated;
+				PRAGMA_DISABLE_DEPRECATION_WARNINGS;
+				Plugin->AddActionSets(PluginAttachSets_Deprecated);
+				PRAGMA_ENABLE_DEPRECATION_WARNINGS;
+				for (const XrActiveActionSet& ActiveSet : PluginAttachSets_Deprecated)
+				{
+					// TODO?: Log deprecation warning at runtime, since overridden deprecated interface methods don't warn at compile time?
+					BindActionSets.Add(ActiveSet.actionSet);
+				}
+			}
+
 			for (auto&& BindActionSet : PluginActionSets)
-				BindActionSets.Add(BindActionSet.actionSet);
+				BindActionSets.Add(BindActionSet);
 
 			XrSessionActionSetsAttachInfo SessionActionSetsAttachInfo;
 			SessionActionSetsAttachInfo.type = XR_TYPE_SESSION_ACTION_SETS_ATTACH_INFO;
@@ -569,16 +651,22 @@ void FOpenXRInputPlugin::FOpenXRInput::Tick(float DeltaTime)
 		SyncInfo.type = XR_TYPE_ACTIONS_SYNC_INFO;
 		SyncInfo.next = nullptr;
 
-		PluginActionSets.Empty();
+		TArray<XrActiveActionSet> ActiveSets;
 		for (IOpenXRExtensionPlugin* Plugin : OpenXRHMD->GetExtensionPlugins())
 		{
-			Plugin->AddActionSets(PluginActionSets);
+			Plugin->GetActiveActionSetsForSync(ActiveSets);
+
+			// TODO?: Log deprecation warning at runtime, since overridden deprecated interface methods don't warn at compile time?
+			PRAGMA_DISABLE_DEPRECATION_WARNINGS;
+			Plugin->AddActionSets(ActiveSets);
+			PRAGMA_ENABLE_DEPRECATION_WARNINGS;
+
 			SyncInfo.next = Plugin->OnSyncActions(Session, SyncInfo.next);
 		}
 
-		PluginActionSets.Append(ActionSets);
-		SyncInfo.countActiveActionSets = PluginActionSets.Num();
-		SyncInfo.activeActionSets = PluginActionSets.GetData();
+		ActiveSets.Append(ActionSets);
+		SyncInfo.countActiveActionSets = ActiveSets.Num();
+		SyncInfo.activeActionSets = ActiveSets.GetData();
 
 		XR_ENSURE(xrSyncActions(Session, &SyncInfo));
 	
diff --git a/Engine/Plugins/Runtime/OpenXR/Source/OpenXRInput/Private/OpenXRInput.h b/Engine/Plugins/Runtime/OpenXR/Source/OpenXRInput/Private/OpenXRInput.h
--- a/Engine/Plugins/Runtime/OpenXR/Source/OpenXRInput/Private/OpenXRInput.h
+++ b/Engine/Plugins/Runtime/OpenXR/Source/OpenXRInput/Private/OpenXRInput.h
@@ -95,7 +95,7 @@ public:
 		FOpenXRHMD* OpenXRHMD;
 
 		TArray<XrActiveActionSet> ActionSets;
-		TArray<XrActiveActionSet> PluginActionSets;
+		TArray<XrActionSet> PluginActionSets;
 		TArray<XrPath> SubactionPaths;
 		TArray<FOpenXRAction> Actions;
 		TMap<EControllerHand, FOpenXRController> Controllers;
@@ -114,6 +114,8 @@ public:
 		template<typename T>
 		int32 SuggestBindings(XrInstance Instance, FOpenXRAction& Action, const TArray<T>& Mappings, TMap<FString, FInteractionProfile>& Profiles);
 
+		bool SuggestBindingForKey(XrInstance Instance, FOpenXRAction& Action, const FKey& Key, TMap<FString, FInteractionProfile>& Profiles);
+
 		/** handler to send all messages to */
 		TSharedRef<FGenericApplicationMessageHandler> MessageHandler;
 	};
