OpenXR: properly initialize pipelined XrSpaceLocation

We were using SetNum for our DeviceLocations TArray. This lead to garbage data in DeviceLocations, which would accidentally populate LocationFlags.
Because we rely on caching for 'last known' locations, this leads to bugs when we fetch a partially valid location.
One specific bug scenario: we fetch a valid position, but invalid orientation. We reset position correctly in the cached locations. But the garbage data
is left for orientation, leading to failures in GetCurrentPose.

diff --git a/Engine/Plugins/Runtime/OpenXR/Source/OpenXRHMD/Private/OpenXRHMD.cpp b/Engine/Plugins/Runtime/OpenXR/Source/OpenXRHMD/Private/OpenXRHMD.cpp
--- a/Engine/Plugins/Runtime/OpenXR/Source/OpenXRHMD/Private/OpenXRHMD.cpp
+++ b/Engine/Plugins/Runtime/OpenXR/Source/OpenXRHMD/Private/OpenXRHMD.cpp
@@ -1081,7 +1081,7 @@ void FOpenXRHMD::UpdateDeviceLocations(bool bUpdateOpenXRExtensionPlugins)
 	if (PipelineState.FrameState.predictedDisplayTime > 0)
 	{
 		FReadScopeLock Lock(DeviceMutex);
-		PipelineState.DeviceLocations.SetNum(DeviceSpaces.Num());
+		PipelineState.DeviceLocations.SetNumZeroed(DeviceSpaces.Num());
 		for (int32 DeviceIndex = 0; DeviceIndex < PipelineState.DeviceLocations.Num(); DeviceIndex++)
 		{
 			const FDeviceSpace& DeviceSpace = DeviceSpaces[DeviceIndex];
