Fix invalid extension name pointers when querying extensions for multiple devices

diff --git a/Engine/Plugins/Runtime/OpenXR/Source/OpenXRHMD/Private/OpenXRHMD.cpp b/Engine/Plugins/Runtime/OpenXR/Source/OpenXRHMD/Private/OpenXRHMD.cpp
--- a/Engine/Plugins/Runtime/OpenXR/Source/OpenXRHMD/Private/OpenXRHMD.cpp
+++ b/Engine/Plugins/Runtime/OpenXR/Source/OpenXRHMD/Private/OpenXRHMD.cpp
@@ -93,6 +93,7 @@ bool FOpenXRHMD::FVulkanExtensions::GetVulkanInstanceExtensionsRequired(TArray<c
 #ifdef XR_USE_GRAPHICS_API_VULKAN
 	TArray<VkExtensionProperties> Properties = GetIVulkanDynamicRHI()->RHIGetAllInstanceExtensions();
 
+	if (Extensions.IsEmpty())
 	{
 		PFN_xrGetVulkanInstanceExtensionsKHR GetVulkanInstanceExtensionsKHR;
 		XR_ENSURE(xrGetInstanceProcAddr(Instance, "xrGetVulkanInstanceExtensionsKHR", (PFN_xrVoidFunction*)&GetVulkanInstanceExtensionsKHR));
@@ -134,6 +135,7 @@ bool FOpenXRHMD::FVulkanExtensions::GetVulkanDeviceExtensionsRequired(VkPhysical
 #ifdef XR_USE_GRAPHICS_API_VULKAN
 	TArray<VkExtensionProperties> Properties = GetIVulkanDynamicRHI()->RHIGetAllDeviceExtensions((VkPhysicalDevice)pPhysicalDevice);
 
+	if (DeviceExtensions.IsEmpty())
 	{
 		PFN_xrGetVulkanDeviceExtensionsKHR GetVulkanDeviceExtensionsKHR;
 		XR_ENSURE(xrGetInstanceProcAddr(Instance, "xrGetVulkanDeviceExtensionsKHR", (PFN_xrVoidFunction*)&GetVulkanDeviceExtensionsKHR));
@@ -1081,7 +1083,7 @@ void FOpenXRHMD::UpdateDeviceLocations(bool bUpdateOpenXRExtensionPlugins)
 	if (PipelineState.FrameState.predictedDisplayTime > 0)
 	{
 		FReadScopeLock Lock(DeviceMutex);
-		PipelineState.DeviceLocations.SetNumZeroed(DeviceSpaces.Num());
+		PipelineState.DeviceLocations.SetNum(DeviceSpaces.Num());
 		for (int32 DeviceIndex = 0; DeviceIndex < PipelineState.DeviceLocations.Num(); DeviceIndex++)
 		{
 			const FDeviceSpace& DeviceSpace = DeviceSpaces[DeviceIndex];
