Add a CVar to override the OpenXR Environment Blend mode used for submission

The also changes the clear value to always be transparent, but this seems to be the UE5 default now anyway.

diff --git a/Engine/Plugins/Runtime/OpenXR/Source/OpenXRHMD/Private/OpenXRHMD.cpp b/Engine/Plugins/Runtime/OpenXR/Source/OpenXRHMD/Private/OpenXRHMD.cpp
--- a/Engine/Plugins/Runtime/OpenXR/Source/OpenXRHMD/Private/OpenXRHMD.cpp
+++ b/Engine/Plugins/Runtime/OpenXR/Source/OpenXRHMD/Private/OpenXRHMD.cpp
@@ -46,6 +46,12 @@ static TAutoConsoleVariable<int32> CVarOpenXRExitAppOnRuntimeDrivenSessionExit(
 	TEXT("The aniticipated situation is that the runtime is associated with a launcher application or has a runtime UI overlay which can tell openxr to exit vr and that in that context the app should also exit.  But maybe there are cases where it should not?  Set this CVAR to make it not.\n"),
 	ECVF_Default);
 
+static TAutoConsoleVariable<int32> CVarOpenXREnvironmentBlendMode(
+	TEXT("xr.OpenXREnvironmentBlendMode"),
+	0,
+	TEXT("Override the XrEnvironmentBlendMode used when submitting frames. 1 = Opaque, 2 = Additive, 3 = Alpha Blend\n"),
+	ECVF_Default);
+
 namespace {
 	static TSet<XrEnvironmentBlendMode> SupportedBlendModes{ XR_ENVIRONMENT_BLEND_MODE_ALPHA_BLEND, XR_ENVIRONMENT_BLEND_MODE_ADDITIVE, XR_ENVIRONMENT_BLEND_MODE_OPAQUE };
 	static TSet<XrViewConfigurationType> SupportedViewConfigurations{ XR_VIEW_CONFIGURATION_TYPE_PRIMARY_STEREO, XR_VIEW_CONFIGURATION_TYPE_PRIMARY_QUAD_VARJO };
@@ -1688,7 +1694,7 @@ bool FOpenXRHMD::AllocateRenderTargetTexture(uint32 Index, uint32 SizeX, uint32
 		Format = PF_R8G8B8A8;
 	}
 
-	FClearValueBinding ClearColor = (SelectedEnvironmentBlendMode == XR_ENVIRONMENT_BLEND_MODE_OPAQUE) ? FClearValueBinding::Black : FClearValueBinding::Transparent;
+	FClearValueBinding ClearColor = FClearValueBinding::Transparent;
 
 	FXRSwapChainPtr& Swapchain = PipelinedLayerStateRendering.ColorSwapchain;
 	const FRHITexture2D* const SwapchainTexture = Swapchain == nullptr ? nullptr : Swapchain->GetTexture2DArray() ? Swapchain->GetTexture2DArray() : Swapchain->GetTexture2D();
@@ -2329,11 +2335,13 @@ void FOpenXRHMD::OnFinishRendering_RHIThread()
 			Headers.Add(reinterpret_cast<const XrCompositionLayerBaseHeader*>(&Quad));
 		}
 
+		int32 BlendModeOverride = CVarOpenXREnvironmentBlendMode.GetValueOnRenderThread();
+
 		XrFrameEndInfo EndInfo;
 		EndInfo.type = XR_TYPE_FRAME_END_INFO;
 		EndInfo.next = nullptr;
 		EndInfo.displayTime = PipelinedFrameStateRHI.FrameState.predictedDisplayTime;
-		EndInfo.environmentBlendMode = SelectedEnvironmentBlendMode;
+		EndInfo.environmentBlendMode = BlendModeOverride ? (XrEnvironmentBlendMode)BlendModeOverride : SelectedEnvironmentBlendMode;
 
 		EndInfo.layerCount = PipelinedFrameStateRHI.FrameState.shouldRender ? Headers.Num() : 0;
 		EndInfo.layers = PipelinedFrameStateRHI.FrameState.shouldRender ? Headers.GetData() : nullptr;
