OpenXR: No more VRPIE exit when stereo disabled

Previous in VRPIE, if we called `vr.bEnableStereo` or used the "Enable HMD" BP node, we'd trigger VRPIE exit.
No longer! We can retain the VRPIE session, and toggle between 2D and stereo in the same PIE session.

diff --git a/Engine/Plugins/Runtime/OpenXR/Source/OpenXRHMD/Private/OpenXRHMD.cpp b/Engine/Plugins/Runtime/OpenXR/Source/OpenXRHMD/Private/OpenXRHMD.cpp
--- a/Engine/Plugins/Runtime/OpenXR/Source/OpenXRHMD/Private/OpenXRHMD.cpp
+++ b/Engine/Plugins/Runtime/OpenXR/Source/OpenXRHMD/Private/OpenXRHMD.cpp
@@ -2142,10 +2142,7 @@ bool FOpenXRHMD::OnStartGameFrame(FWorldContext& WorldContext)
 
 				// But only if this CVar is set to true.
 				bExitApp = bExitApp && (CVarOpenXRExitAppOnRuntimeDrivenSessionExit.GetValueOnAnyThread() != 0);
-#if WITH_EDITOR
-				// But always if in the editor, because that doesn't actually exit.  See RequestExitApp().
-				bExitApp = bExitApp || GIsEditor;
-#endif
+	
 				if (bExitApp)
 				{
 					RequestExitApp();
