OpenXR: don't try to get space accelleration unless the extension is supported.

diff --git a/Engine/Plugins/Runtime/OpenXR/Source/OpenXRHMD/Private/OpenXRHMD.cpp b/Engine/Plugins/Runtime/OpenXR/Source/OpenXRHMD/Private/OpenXRHMD.cpp
--- a/Engine/Plugins/Runtime/OpenXR/Source/OpenXRHMD/Private/OpenXRHMD.cpp
+++ b/Engine/Plugins/Runtime/OpenXR/Source/OpenXRHMD/Private/OpenXRHMD.cpp
@@ -462,7 +462,8 @@ bool FOpenXRHMD::GetPoseForTime(int32 DeviceId, FTimespan Timespan, bool& OutTim
 	const FDeviceSpace& DeviceSpace = DeviceSpaces[DeviceId];
 
 	XrSpaceAccelerationEPIC DeviceAcceleration{ (XrStructureType)XR_TYPE_SPACE_ACCELERATION_EPIC };
-	XrSpaceVelocity DeviceVelocity { XR_TYPE_SPACE_VELOCITY, &DeviceAcceleration };
+	void* DeviceAccellerationPtr = bSpaceAccellerationSupported ? &DeviceAcceleration : nullptr;
+	XrSpaceVelocity DeviceVelocity { XR_TYPE_SPACE_VELOCITY, DeviceAccellerationPtr };
 	XrSpaceLocation DeviceLocation { XR_TYPE_SPACE_LOCATION, &DeviceVelocity };
 
 	XR_ENSURE(xrLocateSpace(DeviceSpace.Space, PipelineState.TrackingSpace, TargetTime, &DeviceLocation));
@@ -927,7 +928,6 @@ FOpenXRHMD::FOpenXRHMD(const FAutoRegister& AutoRegister, XrInstance InInstance,
 	InstanceProperties.runtimeName[XR_MAX_RUNTIME_NAME_SIZE - 1] = 0; // Ensure the name is null terminated.
 
 	bDepthExtensionSupported = IsExtensionEnabled(XR_KHR_COMPOSITION_LAYER_DEPTH_EXTENSION_NAME) && CheckPlatformDepthExtensionSupport(InstanceProperties);
-
 	bHiddenAreaMaskSupported = IsExtensionEnabled(XR_KHR_VISIBILITY_MASK_EXTENSION_NAME) &&
 		!FCStringAnsi::Strstr(InstanceProperties.runtimeName, "Oculus");
 	bViewConfigurationFovSupported = IsExtensionEnabled(XR_EPIC_VIEW_CONFIGURATION_FOV_EXTENSION_NAME);
@@ -939,6 +939,8 @@ FOpenXRHMD::FOpenXRHMD(const FAutoRegister& AutoRegister, XrInstance InInstance,
 	bSupportsHandTracking = HandTrackingSystemProperties.supportsHandTracking == XR_TRUE;
 	SystemProperties.next = nullptr;
 
+	bSpaceAccellerationSupported = IsExtensionEnabled(XR_EPIC_SPACE_ACCELERATION_NAME);
+
 	static const auto CVarMobileMultiView = IConsoleManager::Get().FindTConsoleVariableDataInt(TEXT("vr.MobileMultiView"));
 	static const auto CVarMobileHDR = IConsoleManager::Get().FindTConsoleVariableDataInt(TEXT("r.MobileHDR"));
 	const bool bMobileHDR = (CVarMobileHDR && CVarMobileHDR->GetValueOnAnyThread() != 0);
diff --git a/Engine/Plugins/Runtime/OpenXR/Source/OpenXRHMD/Private/OpenXRHMD.h b/Engine/Plugins/Runtime/OpenXR/Source/OpenXRHMD/Private/OpenXRHMD.h
--- a/Engine/Plugins/Runtime/OpenXR/Source/OpenXRHMD/Private/OpenXRHMD.h
+++ b/Engine/Plugins/Runtime/OpenXR/Source/OpenXRHMD/Private/OpenXRHMD.h
@@ -326,6 +326,7 @@ private:
 	bool					bNeedReBuildOcclusionMesh;
 	bool					bIsMobileMultiViewEnabled;
 	bool					bSupportsHandTracking;
+	bool					bSpaceAccellerationSupported;
 	bool					bProjectionLayerAlphaEnabled;
 	bool					bIsStandaloneStereoOnlyDevice;
 	float					WorldToMetersScale = 100.0f;
diff --git a/Engine/Plugins/Runtime/OpenXR/Source/OpenXRHMD/Private/OpenXRHMDModule.cpp b/Engine/Plugins/Runtime/OpenXR/Source/OpenXRHMD/Private/OpenXRHMDModule.cpp
--- a/Engine/Plugins/Runtime/OpenXR/Source/OpenXRHMD/Private/OpenXRHMDModule.cpp
+++ b/Engine/Plugins/Runtime/OpenXR/Source/OpenXRHMD/Private/OpenXRHMDModule.cpp
@@ -446,6 +446,10 @@ bool FOpenXRHMDModule::GetOptionalExtensions(TArray<const ANSICHAR*>& OutExtensi
 
 	// Draft extension not yet provided in headers
 	OutExtensions.Add("XR_EXT_dpad_binding");
+
+	// A pre-draft extension.
+	OutExtensions.Add(XR_EPIC_VIEW_CONFIGURATION_FOV_EXTENSION_NAME);
+
 	return true;
 }
 
