OpenXR: Expose xrGetCurrentInteractionProfile to blueprint

diff --git a/Engine/Plugins/Runtime/OpenXR/Source/OpenXRHMD/Private/OpenXRHMD.cpp b/Engine/Plugins/Runtime/OpenXR/Source/OpenXRHMD/Private/OpenXRHMD.cpp
--- a/Engine/Plugins/Runtime/OpenXR/Source/OpenXRHMD/Private/OpenXRHMD.cpp
+++ b/Engine/Plugins/Runtime/OpenXR/Source/OpenXRHMD/Private/OpenXRHMD.cpp
@@ -192,18 +192,11 @@ void FOpenXRHMD::GetMotionControllerData(UObject* WorldContext, const EControlle
 		FReadScopeLock SessionLock(SessionHandleMutex);
 		if (Session)
 		{
-			XrInteractionProfileState Profile;
-			Profile.type = XR_TYPE_INTERACTION_PROFILE_STATE;
-			Profile.next = nullptr;
+			XrInteractionProfileState Profile = { XR_TYPE_INTERACTION_PROFILE_STATE };
 			if (XR_SUCCEEDED(xrGetCurrentInteractionProfile(Session, GetTrackedDevicePath(Devices[(int32)Hand]), &Profile)) &&
 				Profile.interactionProfile != XR_NULL_PATH)
 			{
-				TArray<char> Path;
-				uint32 PathCount = 0;
-				XR_ENSURE(xrPathToString(Instance, Profile.interactionProfile, 0, &PathCount, nullptr));
-				Path.SetNum(PathCount);
-				XR_ENSURE(xrPathToString(Instance, Profile.interactionProfile, PathCount, &PathCount, Path.GetData()));
-				MotionControllerData.DeviceName = Path.GetData();
+				XR_ENSURE(OpenXRPathToFName(Instance, Profile.interactionProfile, MotionControllerData.DeviceName));
 			}
 		}
 	}
@@ -273,6 +266,70 @@ void FOpenXRHMD::GetMotionControllerData(UObject* WorldContext, const EControlle
 	MotionControllerData.bIsGrasped = false;
 }
 
+bool FOpenXRHMD::GetCurrentInteractionProfile(const EControllerHand Hand, FString& InteractionProfile)
+{
+	int32 DeviceId = -1;
+	if (Hand == EControllerHand::HMD)
+	{
+		DeviceId = IXRTrackingSystem::HMDDeviceId;
+	}
+	else
+	{
+		TArray<int32> Devices;
+		if (EnumerateTrackedDevices(Devices, EXRTrackedDeviceType::Controller))
+		{
+			if (Devices.IsValidIndex((int32)Hand))
+			{
+				DeviceId = Devices[(int32)Hand];
+			}
+		}
+	}
+
+	if (DeviceId == -1)
+	{
+		UE_LOG(LogHMD, Warning, TEXT("GetCurrentInteractionProfile failed because that EControllerHandValue %i does not map to a device!"), Hand);
+		return false;
+	}
+
+	FReadScopeLock SessionLock(SessionHandleMutex);
+	if (Session)
+	{
+		XrInteractionProfileState Profile{ XR_TYPE_INTERACTION_PROFILE_STATE };
+		XrPath Path = GetTrackedDevicePath(DeviceId);
+		XrResult Result = xrGetCurrentInteractionProfile(Session, Path, &Profile);
+		if (XR_SUCCEEDED(Result))
+		{
+			if (Profile.interactionProfile == XR_NULL_PATH)
+			{
+				InteractionProfile = "";
+				return true;
+			}
+			else
+			{				
+				XR_ENSURE(OpenXRPathToFString(Instance, Profile.interactionProfile, InteractionProfile));
+				return true;
+			}
+		}
+		else
+		{
+			FString PathStr;
+			XrResult PathResult = OpenXRPathToFString(Instance, Path, PathStr);
+			if (!XR_SUCCEEDED(PathResult))
+			{
+				PathStr = FString::Printf(TEXT("xrPathToString returned %s"), OpenXRResultToString(Result));
+			}
+			UE_LOG(LogHMD, Warning, TEXT("GetCurrentInteractionProfile for %i (%s) failed because xrGetCurrentInteractionProfile failed with result %s."), Hand, *PathStr, OpenXRResultToString(Result));
+			return false;
+		}
+	}
+	else
+	{
+		UE_LOG(LogHMD, Warning, TEXT("GetCurrentInteractionProfile for %i failed because session is null!"), Hand);
+		return false;
+	}
+
+}
+
 float FOpenXRHMD::GetWorldToMetersScale() const
 {
 	return IsInActualRenderingThread() ? PipelinedFrameStateRendering.WorldToMetersScale : PipelinedFrameStateGame.WorldToMetersScale;
diff --git a/Engine/Plugins/Runtime/OpenXR/Source/OpenXRHMD/Private/OpenXRHMD.h b/Engine/Plugins/Runtime/OpenXR/Source/OpenXRHMD/Private/OpenXRHMD.h
--- a/Engine/Plugins/Runtime/OpenXR/Source/OpenXRHMD/Private/OpenXRHMD.h
+++ b/Engine/Plugins/Runtime/OpenXR/Source/OpenXRHMD/Private/OpenXRHMD.h
@@ -136,6 +136,7 @@ public:
 	virtual bool GetIsTracked(int32 DeviceId);
 	virtual bool GetCurrentPose(int32 DeviceId, FQuat& CurrentOrientation, FVector& CurrentPosition) override;
 	virtual bool GetPoseForTime(int32 DeviceId, FTimespan Timespan, bool& OutTimeWasUsed, FQuat& CurrentOrientation, FVector& CurrentPosition, bool& bProvidedLinearVelocity, FVector& LinearVelocity, bool& bProvidedAngularVelocity, FVector& AngularVelocityRadPerSec, bool& bProvidedLinearAcceleration, FVector& LinearAcceleration, float WorldToMetersScale);
+	virtual bool GetCurrentInteractionProfile(const EControllerHand Hand, FString& InteractionProfile) override;
 	virtual void SetBaseRotation(const FRotator& BaseRot) override;
 	virtual FRotator GetBaseRotation() const override;
 
diff --git a/Engine/Plugins/Runtime/OpenXR/Source/OpenXRHMD/Public/OpenXRCore.h b/Engine/Plugins/Runtime/OpenXR/Source/OpenXRHMD/Public/OpenXRCore.h
--- a/Engine/Plugins/Runtime/OpenXR/Source/OpenXRHMD/Public/OpenXRCore.h
+++ b/Engine/Plugins/Runtime/OpenXR/Source/OpenXRHMD/Public/OpenXRCore.h
@@ -200,3 +200,35 @@ bool PreInitOpenXRCore(PFN_xrGetInstanceProcAddr InGetProcAddr);
  * @returns true if initialization was successful.
  */
 bool InitOpenXRCore(XrInstance Instance);
+
+FORCEINLINE XrResult OpenXRPathToFString(XrInstance Instance, XrPath Path, FString& OutString)
+{
+	uint32 PathCount = 0;
+	char PathChars[XR_MAX_PATH_LENGTH];
+	XrResult Result = xrPathToString(Instance, Path, XR_MAX_PATH_LENGTH, &PathCount, PathChars);
+	if (Result == XR_SUCCESS)
+	{
+		OutString = FString(PathCount, PathChars);
+	}
+	else
+	{
+		OutString = "";
+	}	
+	return Result;
+}
+
+FORCEINLINE XrResult OpenXRPathToFName(XrInstance Instance, XrPath Path, FName& OutFName)
+{
+	uint32 PathCount = 0;
+	char PathChars[XR_MAX_PATH_LENGTH];
+	XrResult Result = xrPathToString(Instance, Path, XR_MAX_PATH_LENGTH, &PathCount, PathChars);
+	if (Result == XR_SUCCESS)
+	{
+		OutFName = FName(PathCount, PathChars);
+	}
+	else
+	{
+		OutFName = NAME_None;
+	}
+	return Result;
+}
\ No newline at end of file
diff --git a/Engine/Source/Runtime/HeadMountedDisplay/Private/HeadMountedDisplayFunctionLibrary.cpp b/Engine/Source/Runtime/HeadMountedDisplay/Private/HeadMountedDisplayFunctionLibrary.cpp
--- a/Engine/Source/Runtime/HeadMountedDisplay/Private/HeadMountedDisplayFunctionLibrary.cpp
+++ b/Engine/Source/Runtime/HeadMountedDisplay/Private/HeadMountedDisplayFunctionLibrary.cpp
@@ -487,6 +487,16 @@ bool UHeadMountedDisplayFunctionLibrary::ConfigureGestures(const FXRGestureConfi
 	return false;
 }
 
+bool UHeadMountedDisplayFunctionLibrary::GetCurrentInteractionProfile(const EControllerHand Hand, FString& InteractionProfile)
+{
+	IXRTrackingSystem* TrackingSys = GEngine->XRSystem.Get();
+	if (TrackingSys)
+	{
+		return TrackingSys->GetCurrentInteractionProfile(Hand, InteractionProfile);
+	}
+
+	return false;
+}
 
 /** Connect to a remote device for Remote Debugging */
 EXRDeviceConnectionResult::Type UHeadMountedDisplayFunctionLibrary::ConnectRemoteXRDevice(const FString& IpAddress, const int32 BitRate)
diff --git a/Engine/Source/Runtime/HeadMountedDisplay/Private/XRTrackingSystemBase.cpp b/Engine/Source/Runtime/HeadMountedDisplay/Private/XRTrackingSystemBase.cpp
--- a/Engine/Source/Runtime/HeadMountedDisplay/Private/XRTrackingSystemBase.cpp
+++ b/Engine/Source/Runtime/HeadMountedDisplay/Private/XRTrackingSystemBase.cpp
@@ -163,6 +163,13 @@ void FXRTrackingSystemBase::GetMotionControllerData(UObject* WorldContext, const
 	}
 }
 
+bool FXRTrackingSystemBase::GetCurrentInteractionProfile(const EControllerHand Hand, FString& InteractionProfile)
+{
+	UE_LOG(LogHMD, Warning, TEXT("GetCurrentInteractionProfile for %i failed because the current VR tracking system does not support it!"), Hand);
+
+	return false;
+}
+
 
 class IXRLoadingScreen* FXRTrackingSystemBase::GetLoadingScreen()
 {
diff --git a/Engine/Source/Runtime/HeadMountedDisplay/Public/HeadMountedDisplayFunctionLibrary.h b/Engine/Source/Runtime/HeadMountedDisplay/Public/HeadMountedDisplayFunctionLibrary.h
--- a/Engine/Source/Runtime/HeadMountedDisplay/Public/HeadMountedDisplayFunctionLibrary.h
+++ b/Engine/Source/Runtime/HeadMountedDisplay/Public/HeadMountedDisplayFunctionLibrary.h
@@ -328,10 +328,15 @@ class HEADMOUNTEDDISPLAY_API UHeadMountedDisplayFunctionLibrary : public UBluepr
 	UFUNCTION(BlueprintCallable, Category = "Input|XRTracking")
 	static void GetMotionControllerData(UObject* WorldContext, const EControllerHand Hand, FXRMotionControllerData& MotionControllerData);
 
-
 	UFUNCTION(BlueprintCallable, Category = "Input|XRTracking", meta = (ToolTip = "Specify which gestures to capture."))
 	static bool ConfigureGestures(const FXRGestureConfig& GestureConfig);
 
+	/**
+	 * Get the openXR interaction profile name for the given controller. Returns true if the openxr call is successfully made.  The string may be empty
+	 * if there is no interaction profile associated with the controller.
+	 */	
+	UFUNCTION(BlueprintCallable, Category = "Input|XRTracking|OpenXR")
+	static bool GetCurrentInteractionProfile(const EControllerHand Hand, FString& InteractionProfile);
 	
 	/** Connect to a remote device */
 	UFUNCTION(BlueprintCallable, Category = "XR|HeadMountedDisplay")
diff --git a/Engine/Source/Runtime/HeadMountedDisplay/Public/IXRTrackingSystem.h b/Engine/Source/Runtime/HeadMountedDisplay/Public/IXRTrackingSystem.h
--- a/Engine/Source/Runtime/HeadMountedDisplay/Public/IXRTrackingSystem.h
+++ b/Engine/Source/Runtime/HeadMountedDisplay/Public/IXRTrackingSystem.h
@@ -422,6 +422,7 @@ public:
 	 * Platform Agnostic Query about MotionControllers details
 	 */
 	virtual void GetMotionControllerData(UObject* WorldContext, const EControllerHand Hand, FXRMotionControllerData& MotionControllerData) = 0;
+	virtual bool GetCurrentInteractionProfile(const EControllerHand Hand, FString& InteractionProfile) = 0;
 
 	virtual bool ConfigureGestures(const FXRGestureConfig& GestureConfig) = 0;
 
diff --git a/Engine/Source/Runtime/HeadMountedDisplay/Public/XRTrackingSystemBase.h b/Engine/Source/Runtime/HeadMountedDisplay/Public/XRTrackingSystemBase.h
--- a/Engine/Source/Runtime/HeadMountedDisplay/Public/XRTrackingSystemBase.h
+++ b/Engine/Source/Runtime/HeadMountedDisplay/Public/XRTrackingSystemBase.h
@@ -161,6 +161,7 @@ public:
 	virtual class IXRLoadingScreen* GetLoadingScreen() override final;
 
 	virtual void GetMotionControllerData(UObject* WorldContext, const EControllerHand Hand, FXRMotionControllerData& MotionControllerData) override;
+	virtual bool GetCurrentInteractionProfile(const EControllerHand Hand, FString& InteractionProfile) override;
 
 	virtual bool ConfigureGestures(const FXRGestureConfig& GestureConfig) override { return false; }
 
