OpenXRHMD initialize space locations to avoid bad quaternions.
-Was getting bad rotator asserts in debug build on oculus with openxr.  Turns out we had some uninitialized XrSpaceLocation structs if tracking is not yet succeeding.

diff --git a/Engine/Plugins/Runtime/OpenXR/Source/OpenXRHMD/Private/OpenXRHMD.cpp b/Engine/Plugins/Runtime/OpenXR/Source/OpenXRHMD/Private/OpenXRHMD.cpp
--- a/Engine/Plugins/Runtime/OpenXR/Source/OpenXRHMD/Private/OpenXRHMD.cpp
+++ b/Engine/Plugins/Runtime/OpenXR/Source/OpenXRHMD/Private/OpenXRHMD.cpp
@@ -1148,14 +1148,22 @@ void FOpenXRHMD::UpdateDeviceLocations(bool bUpdateOpenXRExtensionPlugins)
 	if (PipelineState.FrameState.predictedDisplayTime > 0)
 	{
 		FReadScopeLock Lock(DeviceMutex);
-		PipelineState.DeviceLocations.SetNum(DeviceSpaces.Num());
+		const int NeededElements = DeviceSpaces.Num() - PipelineState.DeviceLocations.Num();
+		check(NeededElements >= 0);
+		if (NeededElements > 0)
+		{
+			PipelineState.DeviceLocations.AddZeroed(NeededElements);
+			for (XrSpaceLocation& Space : PipelineState.DeviceLocations)
+			{
+				Space.type = XR_TYPE_SPACE_LOCATION;
+			}
+		}
 		for (int32 DeviceIndex = 0; DeviceIndex < PipelineState.DeviceLocations.Num(); DeviceIndex++)
 		{
 			const FDeviceSpace& DeviceSpace = DeviceSpaces[DeviceIndex];
 			if (DeviceSpace.Space != XR_NULL_HANDLE)
 			{
-				XrSpaceLocation NewDeviceLocation = {};
-				NewDeviceLocation.type = XR_TYPE_SPACE_LOCATION;
+				XrSpaceLocation NewDeviceLocation = {XR_TYPE_SPACE_LOCATION};
 				XrResult Result = xrLocateSpace(DeviceSpace.Space, PipelineState.TrackingSpace, PipelineState.FrameState.predictedDisplayTime, &NewDeviceLocation);
 				if (Result == XR_ERROR_TIME_INVALID)
 				{
diff --git a/Engine/Plugins/Runtime/OpenXR/Source/OpenXRHMD/Private/OpenXRHMD.h b/Engine/Plugins/Runtime/OpenXR/Source/OpenXRHMD/Private/OpenXRHMD.h
--- a/Engine/Plugins/Runtime/OpenXR/Source/OpenXRHMD/Private/OpenXRHMD.h
+++ b/Engine/Plugins/Runtime/OpenXR/Source/OpenXRHMD/Private/OpenXRHMD.h
@@ -58,7 +58,7 @@ public:
 		XrViewState ViewState{XR_TYPE_VIEW_STATE};
 		TArray<XrView> Views;
 		TArray<XrSpaceLocation> DeviceLocations;
-		XrSpace TrackingSpace;
+		XrSpace TrackingSpace = XR_NULL_HANDLE;
 		float WorldToMetersScale = 100.0f;
 
 		TArray<XrViewConfigurationView> ViewConfigs;
