OpenXR: Add OnDestroySession to IOpenXRExtensionPlugin.

diff --git a/Engine/Plugins/Runtime/OpenXR/Source/OpenXRHMD/Private/OpenXRHMD.cpp b/Engine/Plugins/Runtime/OpenXR/Source/OpenXRHMD/Private/OpenXRHMD.cpp
--- a/Engine/Plugins/Runtime/OpenXR/Source/OpenXRHMD/Private/OpenXRHMD.cpp
+++ b/Engine/Plugins/Runtime/OpenXR/Source/OpenXRHMD/Private/OpenXRHMD.cpp
@@ -1562,12 +1562,18 @@ void FOpenXRHMD::DestroySession()
 {
 	// FlushRenderingCommands must be called outside of SessionLock since some rendering threads will also lock this mutex.
 	FlushRenderingCommands();
+
 	FWriteScopeLock SessionLock(SessionHandleMutex);
 
 	if (Session != XR_NULL_HANDLE)
 	{
 		FWriteScopeLock DeviceLock(DeviceMutex);
 
+		for (IOpenXRExtensionPlugin* Module : ExtensionPlugins)
+		{
+			Module->OnDestroySession(Session);
+		}
+
 		// We need to reset all swapchain references to ensure there are no attempts
 		// to destroy swapchain handles after the session is already destroyed.
 		ForEachLayer([&](uint32 /* unused */, FOpenXRLayer& Layer)
diff --git a/Engine/Plugins/Runtime/OpenXR/Source/OpenXRHMD/Public/IOpenXRExtensionPlugin.h b/Engine/Plugins/Runtime/OpenXR/Source/OpenXRHMD/Public/IOpenXRExtensionPlugin.h
--- a/Engine/Plugins/Runtime/OpenXR/Source/OpenXRHMD/Public/IOpenXRExtensionPlugin.h
+++ b/Engine/Plugins/Runtime/OpenXR/Source/OpenXRHMD/Public/IOpenXRExtensionPlugin.h
@@ -349,6 +349,10 @@ public:
 		return InNext;
 	}
 
+	virtual void OnDestroySession(XrSession InSession)
+	{
+	}
+
 	// OpenXRHMD::OnBeginRendering_GameThread
 	virtual void* OnWaitFrame(XrSession InSession, void* InNext)
 	{
