OpenXR DeviceLocations fix.
-Turns out the array can be shrunk between PIE sessions.  Because we check if the index is in the valid range we do want it to actually shrink there.

diff --git a/Engine/Plugins/Runtime/OpenXR/Source/OpenXRHMD/Private/OpenXRHMD.cpp b/Engine/Plugins/Runtime/OpenXR/Source/OpenXRHMD/Private/OpenXRHMD.cpp
--- a/Engine/Plugins/Runtime/OpenXR/Source/OpenXRHMD/Private/OpenXRHMD.cpp
+++ b/Engine/Plugins/Runtime/OpenXR/Source/OpenXRHMD/Private/OpenXRHMD.cpp
@@ -1148,22 +1148,16 @@
 	if (PipelineState.FrameState.predictedDisplayTime > 0)
 	{
 		FReadScopeLock Lock(DeviceMutex);
-		const int NeededElements = DeviceSpaces.Num() - PipelineState.DeviceLocations.Num();
-		check(NeededElements >= 0);
-		if (NeededElements > 0)
-		{
-			PipelineState.DeviceLocations.AddZeroed(NeededElements);
-			for (XrSpaceLocation& Space : PipelineState.DeviceLocations)
-			{
-				Space.type = XR_TYPE_SPACE_LOCATION;
-			}
-		}
+		PipelineState.DeviceLocations.SetNumZeroed(DeviceSpaces.Num());
 		for (int32 DeviceIndex = 0; DeviceIndex < PipelineState.DeviceLocations.Num(); DeviceIndex++)
 		{
 			const FDeviceSpace& DeviceSpace = DeviceSpaces[DeviceIndex];
+			XrSpaceLocation& CachedDeviceLocation = PipelineState.DeviceLocations[DeviceIndex];
+			CachedDeviceLocation.type = XR_TYPE_SPACE_LOCATION;
+
 			if (DeviceSpace.Space != XR_NULL_HANDLE)
 			{
-				XrSpaceLocation NewDeviceLocation = {XR_TYPE_SPACE_LOCATION};
+				XrSpaceLocation NewDeviceLocation = { XR_TYPE_SPACE_LOCATION };
 				XrResult Result = xrLocateSpace(DeviceSpace.Space, PipelineState.TrackingSpace, PipelineState.FrameState.predictedDisplayTime, &NewDeviceLocation);
 				if (Result == XR_ERROR_TIME_INVALID)
 				{
@@ -1175,7 +1169,6 @@
 					XR_ENSURE(Result);
 				}
 				
-				XrSpaceLocation& CachedDeviceLocation = PipelineState.DeviceLocations[DeviceIndex];
 				// Clear the location tracked bits
 				CachedDeviceLocation.locationFlags &= ~(XR_SPACE_LOCATION_POSITION_TRACKED_BIT | XR_SPACE_LOCATION_ORIENTATION_TRACKED_BIT);
 				if (NewDeviceLocation.locationFlags & (XR_SPACE_LOCATION_POSITION_VALID_BIT))
@@ -1192,7 +1185,7 @@
 			else
 			{
 				// Ensure the location flags are zeroed out so the pose is detected as invalid
-				PipelineState.DeviceLocations[DeviceIndex].locationFlags = 0;
+				CachedDeviceLocation.locationFlags = 0;
 			}
 		}
 
