Allow 10bit swapchain on devices that support it

diff --git a/Engine/Plugins/Runtime/OpenXR/Source/OpenXRHMD/Private/OpenXRHMD.cpp b/Engine/Plugins/Runtime/OpenXR/Source/OpenXRHMD/Private/OpenXRHMD.cpp
--- a/Engine/Plugins/Runtime/OpenXR/Source/OpenXRHMD/Private/OpenXRHMD.cpp
+++ b/Engine/Plugins/Runtime/OpenXR/Source/OpenXRHMD/Private/OpenXRHMD.cpp
@@ -1755,8 +1755,9 @@
 
 
 	// Temporary workaround to swapchain formats - OpenXR doesn't support 10-bit sRGB swapchains, so prefer 8-bit sRGB instead.
-	if (Format == PF_A2B10G10R10)
+	if (Format == PF_A2B10G10R10 && !RenderBridge->Support10BitSwapchain())
 	{
+		UE_LOG(LogHMD, Warning, TEXT("Requesting 10 bit swapchain, but not supported: fall back to 8bpc"));
 		Format = PF_R8G8B8A8;
 	}
 
@@ -1880,7 +1881,7 @@
 	auto CreateSwapchain = [this](FRHITexture2D* Texture, ETextureCreateFlags Flags)
 	{
 		return RenderBridge->CreateSwapchain(Session,
-			PF_B8G8R8A8,
+			IStereoRenderTargetManager::GetStereoLayerPixelFormat(),
 			Texture->GetSizeX(),
 			Texture->GetSizeY(),
 			1,
diff --git a/Engine/Plugins/Runtime/OpenXR/Source/OpenXRHMD/Private/OpenXRHMD_RenderBridge.h b/Engine/Plugins/Runtime/OpenXR/Source/OpenXRHMD/Private/OpenXRHMD_RenderBridge.h
--- a/Engine/Plugins/Runtime/OpenXR/Source/OpenXRHMD/Private/OpenXRHMD_RenderBridge.h
+++ b/Engine/Plugins/Runtime/OpenXR/Source/OpenXRHMD/Private/OpenXRHMD_RenderBridge.h
@@ -48,6 +48,8 @@ public:
 	/** FRHICustomPresent */
 	virtual bool Present(int32& InOutSyncInterval) override;
 
+	virtual bool Support10BitSwapchain() const { return false; }
+
 protected:
 	uint64 AdapterLuid;
 
diff --git a/Engine/Source/Runtime/Engine/Public/StereoRenderTargetManager.h b/Engine/Source/Runtime/Engine/Public/StereoRenderTargetManager.h
--- a/Engine/Source/Runtime/Engine/Public/StereoRenderTargetManager.h
+++ b/Engine/Source/Runtime/Engine/Public/StereoRenderTargetManager.h
@@ -83,4 +83,6 @@ public:
 	 * @return				true, if texture was allocated; false, if the default texture allocation should be used.
 	 */
 	virtual bool AllocateShadingRateTexture(uint32 Index, uint32 RenderSizeX, uint32 RenderSizeY, uint8 Format, uint32 NumMips, ETextureCreateFlags Flags, ETextureCreateFlags TargetableTextureFlags, FTexture2DRHIRef& OutTexture, FIntPoint& OutTextureSize) { return false; }
+
+	static EPixelFormat GetStereoLayerPixelFormat() { return PF_B8G8R8A8; }
 };
\ No newline at end of file
