Implemented FOpenXRHMD::ResetOrientation() and ResetPosition(), fixing UHeadMountedDisplayFunctionLibrary::ResetOrientationAndPosition() behavior when using OpenXR.

diff --git a/Engine/Plugins/Runtime/OpenXR/Source/OpenXRHMD/Private/OpenXRHMD.cpp b/Engine/Plugins/Runtime/OpenXR/Source/OpenXRHMD/Private/OpenXRHMD.cpp
--- a/Engine/Plugins/Runtime/OpenXR/Source/OpenXRHMD/Private/OpenXRHMD.cpp
+++ b/Engine/Plugins/Runtime/OpenXR/Source/OpenXRHMD/Private/OpenXRHMD.cpp
@@ -523,7 +523,7 @@ bool FOpenXRHMD::GetPoseForTime(int32 DeviceId, FTimespan Timespan, bool& OutTim
 	XrSpaceVelocity DeviceVelocity { XR_TYPE_SPACE_VELOCITY, DeviceAccellerationPtr };
 	XrSpaceLocation DeviceLocation { XR_TYPE_SPACE_LOCATION, &DeviceVelocity };
 
-	XR_ENSURE(xrLocateSpace(DeviceSpace.Space, PipelineState.TrackingSpace, TargetTime, &DeviceLocation));
+	XR_ENSURE(xrLocateSpace(DeviceSpace.Space, PipelineState.TrackingSpace->Handle, TargetTime, &DeviceLocation));
 
 	bool ReturnValue = false;
 
@@ -561,36 +561,97 @@ bool FOpenXRHMD::IsChromaAbCorrectionEnabled() const
 	return false;
 }
 
-void FOpenXRHMD::ResetOrientationAndPosition(float yaw)
+void FOpenXRHMD::ResetOrientationAndPosition(float Yaw)
 {
-	ResetOrientation(yaw);
-	ResetPosition();
+	Recenter(EOrientPositionSelector::OrientationAndPosition, Yaw);
 }
 
 void FOpenXRHMD::ResetOrientation(float Yaw)
 {
+	Recenter(EOrientPositionSelector::Orientation, Yaw);
 }
 
 void FOpenXRHMD::ResetPosition()
 {
+
+	Recenter(EOrientPositionSelector::Position);
 }
 
-void FOpenXRHMD::SetBaseRotation(const FRotator& BaseRot)
+void FOpenXRHMD::Recenter(EOrientPositionSelector::Type Selector, float Yaw)
 {
+	FPipelinedFrameState& PipelineState = GetPipelinedFrameStateForThread();
+	XrTime TargetTime = PipelineState.FrameState.predictedDisplayTime;
+
+	FDeviceSpace& DeviceSpace = DeviceSpaces[HMDDeviceId];
+	XrSpaceLocation DeviceLocation = { XR_TYPE_SPACE_LOCATION, nullptr };
+
+	XrSpace BaseSpace = TrackingSpaceType == XR_REFERENCE_SPACE_TYPE_STAGE ? StageSpace : LocalSpace;
+	if (bUseCustomReferenceSpace)
+	{
+		BaseSpace = CustomSpace;
+	}
+	XR_ENSURE(xrLocateSpace(DeviceSpace.Space, BaseSpace, TargetTime, &DeviceLocation));
+
+	const FQuat CurrentOrientation = ToFQuat(DeviceLocation.pose.orientation);
+	const FVector CurrentPosition = ToFVector(DeviceLocation.pose.position, GetWorldToMetersScale());
+
+	if (Selector == EOrientPositionSelector::Position ||
+		Selector == EOrientPositionSelector::OrientationAndPosition)
+	{
+		FVector NewPosition;
+		NewPosition.X = CurrentPosition.X;
+		NewPosition.Y = CurrentPosition.Y;
+		if (TrackingSpaceType == XR_REFERENCE_SPACE_TYPE_LOCAL)
+		{
+			NewPosition.Z = CurrentPosition.Z;
+		}
+		else
+		{
+			NewPosition.Z = 0.0f;
+		}
+		SetBasePosition(NewPosition);
+	}
+
+	if (Selector == EOrientPositionSelector::Orientation ||
+		Selector == EOrientPositionSelector::OrientationAndPosition)
+	{
+		FRotator NewOrientation(0.0f, CurrentOrientation.Rotator().Yaw - Yaw, 0.0f);
+		SetBaseOrientation(NewOrientation.Quaternion());
+	}
+
+	bTrackingSpaceInvalid = true;
+}
+
+void FOpenXRHMD::SetBaseRotation(const FRotator& InBaseRotation)
+{
+	SetBaseOrientation(InBaseRotation.Quaternion());
 }
 
 FRotator FOpenXRHMD::GetBaseRotation() const
 {
-	return FRotator::ZeroRotator;
+	return BaseOrientation.Rotator();
 }
 
-void FOpenXRHMD::SetBaseOrientation(const FQuat& BaseOrient)
+void FOpenXRHMD::SetBaseOrientation(const FQuat& InBaseOrientation)
 {
+	BaseOrientation = InBaseOrientation;
+	bTrackingSpaceInvalid = true;
 }
 
 FQuat FOpenXRHMD::GetBaseOrientation() const
 {
-	return FQuat::Identity;
+	return BaseOrientation;
+}
+
+void FOpenXRHMD::SetBasePosition(const FVector& InBasePosition)
+{
+	BasePosition = InBasePosition;
+	bTrackingSpaceInvalid = true;
+}
+
+FVector FOpenXRHMD::GetBasePosition() const
+{
+	return BasePosition;
 }
 
 bool FOpenXRHMD::IsStereoEnabled() const
@@ -979,6 +1040,10 @@ FOpenXRHMD::FOpenXRHMD(const FAutoRegister& AutoRegister, XrInstance InInstance,
 	, RendererModule(nullptr)
 	, LastRequestedSwapchainFormat(0)
 	, LastRequestedDepthSwapchainFormat(0)
+	, bTrackingSpaceInvalid(true)
+	, bUseCustomReferenceSpace(false)
+	, BaseOrientation(FQuat::Identity)
+	, BasePosition(FVector::ZeroVector)
 {
 	InstanceProperties = { XR_TYPE_INSTANCE_PROPERTIES, nullptr };
 	XR_ENSURE(xrGetInstanceProperties(Instance, &InstanceProperties));
@@ -1158,7 +1223,7 @@ void FOpenXRHMD::UpdateDeviceLocations(bool bUpdateOpenXRExtensionPlugins)
 			if (DeviceSpace.Space != XR_NULL_HANDLE)
 			{
 				XrSpaceLocation NewDeviceLocation = { XR_TYPE_SPACE_LOCATION };
-				XrResult Result = xrLocateSpace(DeviceSpace.Space, PipelineState.TrackingSpace, PipelineState.FrameState.predictedDisplayTime, &NewDeviceLocation);
+				XrResult Result = xrLocateSpace(DeviceSpace.Space, PipelineState.TrackingSpace->Handle, PipelineState.FrameState.predictedDisplayTime, &NewDeviceLocation);
 				if (Result == XR_ERROR_TIME_INVALID)
 				{
 					// The display time is no longer valid so set the location as invalid as well
@@ -1193,7 +1258,7 @@ void FOpenXRHMD::UpdateDeviceLocations(bool bUpdateOpenXRExtensionPlugins)
 		{
 			for (IOpenXRExtensionPlugin* Module : ExtensionPlugins)
 			{
-				Module->UpdateDeviceLocations(Session, PipelineState.FrameState.predictedDisplayTime, PipelineState.TrackingSpace);
+				Module->UpdateDeviceLocations(Session, PipelineState.FrameState.predictedDisplayTime, PipelineState.TrackingSpace->Handle);
 			}
 		}
 	}
@@ -1429,18 +1494,17 @@ bool FOpenXRHMD::OnStereoStartup()
 	uint32_t ReferenceSpacesCount;
 	XR_ENSURE(xrEnumerateReferenceSpaces(Session, 0, &ReferenceSpacesCount, nullptr));
 
-	TArray<XrReferenceSpaceType> Spaces;
-	Spaces.SetNum(ReferenceSpacesCount);
+	TArray<XrReferenceSpaceType> ReferenceSpaces;
+	ReferenceSpaces.SetNum(ReferenceSpacesCount);
 	// Initialize spaces array with valid enum values (avoid triggering validation error).
-	for (auto & SpaceIter : Spaces)
+	for (auto & SpaceIter : ReferenceSpaces)
 		SpaceIter = XR_REFERENCE_SPACE_TYPE_VIEW;
-	XR_ENSURE(xrEnumerateReferenceSpaces(Session, (uint32_t)Spaces.Num(), &ReferenceSpacesCount, Spaces.GetData()));
-	ensure(ReferenceSpacesCount == Spaces.Num());
+	XR_ENSURE(xrEnumerateReferenceSpaces(Session, (uint32_t)ReferenceSpaces.Num(), &ReferenceSpacesCount, ReferenceSpaces.GetData()));
+	ensure(ReferenceSpacesCount == ReferenceSpaces.Num());
 
 	XrSpace HmdSpace = XR_NULL_HANDLE;
 	XrReferenceSpaceCreateInfo SpaceInfo;
-
-	ensure(Spaces.Contains(XR_REFERENCE_SPACE_TYPE_VIEW));
+	ensure(ReferenceSpaces.Contains(XR_REFERENCE_SPACE_TYPE_VIEW));
 	SpaceInfo.type = XR_TYPE_REFERENCE_SPACE_CREATE_INFO;
 	SpaceInfo.next = nullptr;
 	SpaceInfo.referenceSpaceType = XR_REFERENCE_SPACE_TYPE_VIEW;
@@ -1448,30 +1512,30 @@ bool FOpenXRHMD::OnStereoStartup()
 	XR_ENSURE(xrCreateReferenceSpace(Session, &SpaceInfo, &HmdSpace));
 	DeviceSpaces[HMDDeviceId].Space = HmdSpace;
 
-	ensure(Spaces.Contains(XR_REFERENCE_SPACE_TYPE_LOCAL));
+	ensure(ReferenceSpaces.Contains(XR_REFERENCE_SPACE_TYPE_LOCAL));
 	SpaceInfo.referenceSpaceType = XR_REFERENCE_SPACE_TYPE_LOCAL;
 	XR_ENSURE(xrCreateReferenceSpace(Session, &SpaceInfo, &LocalSpace));
 
-	bool UseCustomReferenceSpaceType = false;
+	bUseCustomReferenceSpace = false;
 	XrReferenceSpaceType CustomReferenceSpaceType;
 	for (IOpenXRExtensionPlugin* Module : ExtensionPlugins)
 	{
 		if (Module->UseCustomReferenceSpaceType(CustomReferenceSpaceType))
 		{
-			UseCustomReferenceSpaceType = true;
+			bUseCustomReferenceSpace = true;
 			break;
 		}
 	}
 
-	// If a custom reference space is desired, try to use that.  
-	// Otherwise prefer a stage space over a local space
-	if (UseCustomReferenceSpaceType && Spaces.Contains(CustomReferenceSpaceType))
+	// If a custom reference space is desired, try to use that.
+	// Otherwise use the currently selected reference space.
+	if (bUseCustomReferenceSpace && ReferenceSpaces.Contains(CustomReferenceSpaceType))
 	{
 		TrackingSpaceType = CustomReferenceSpaceType;
 		SpaceInfo.referenceSpaceType = TrackingSpaceType;
 		XR_ENSURE(xrCreateReferenceSpace(Session, &SpaceInfo, &CustomSpace));
 	}
-	else if (Spaces.Contains(XR_REFERENCE_SPACE_TYPE_STAGE))
+	else if (ReferenceSpaces.Contains(XR_REFERENCE_SPACE_TYPE_STAGE))
 	{
 		TrackingSpaceType = XR_REFERENCE_SPACE_TYPE_STAGE;
 		SpaceInfo.referenceSpaceType = TrackingSpaceType;
@@ -1479,9 +1543,16 @@ bool FOpenXRHMD::OnStereoStartup()
 	}
 	else
 	{
+		ensure(ReferenceSpaces.Contains(XR_REFERENCE_SPACE_TYPE_LOCAL));
 		TrackingSpaceType = XR_REFERENCE_SPACE_TYPE_LOCAL;
 	}
 
+	// Create initial tracking space
+	BaseOrientation = FQuat::Identity;
+	BasePosition = FVector::ZeroVector;
+	PipelinedFrameStateGame.TrackingSpace = MakeShared<FTrackingSpace>(TrackingSpaceType);
+	PipelinedFrameStateGame.TrackingSpace->CreateSpace(Session);
+
 	// Create action spaces for all devices
 	for (FDeviceSpace& DeviceSpace : DeviceSpaces)
 	{
@@ -1586,6 +1657,11 @@ void FOpenXRHMD::DestroySession()
 		PipelinedLayerStateRHI.DepthSwapchain.Reset();
 		PipelinedLayerStateRHI.QuadSwapchains.Reset();
 
+		PipelinedFrameStateGame.TrackingSpace.Reset();
+		PipelinedFrameStateRendering.TrackingSpace.Reset();
+		PipelinedFrameStateRHI.TrackingSpace.Reset();
+		bTrackingSpaceInvalid = true;
+
 		// Reset the frame state.
 		PipelinedFrameStateGame.FrameState = XrFrameState{ XR_TYPE_FRAME_STATE };
 		PipelinedFrameStateRendering.FrameState = XrFrameState{ XR_TYPE_FRAME_STATE };
@@ -1595,7 +1671,7 @@ void FOpenXRHMD::DestroySession()
 		FApp::SetUseVRFocus(false);
 		FApp::SetHasVRFocus(false);
 
-		// Destroy device spaces, they will be recreated
+		// Destroy device and reference spaces, they will be recreated
 		// when the session is created again.
 		for (FDeviceSpace& Device : DeviceSpaces)
 		{
@@ -1661,6 +1737,19 @@ XrTime FOpenXRHMD::GetDisplayTime() const
 	return PipelineState.FrameState.predictedDisplayTime;
 }
 
+XrSpace FOpenXRHMD::GetTrackingSpace() const
+{
+	const FPipelinedFrameState& PipelineState = GetPipelinedFrameStateForThread();
+	if (PipelineState.TrackingSpace.IsValid())
+	{
+		return PipelineState.TrackingSpace->Handle;
+	}
+	else
+	{
+		return XR_NULL_HANDLE;
+	}
+}
+
 bool FOpenXRHMD::IsInitialized() const
 {
 	return Instance != XR_NULL_HANDLE;
@@ -1936,7 +2025,7 @@ void FOpenXRHMD::OnBeginRendering_RenderThread(FRHICommandListImmediate& RHICmdL
 		Quad.layerFlags = bNoAlpha ? 0 : XR_COMPOSITION_LAYER_UNPREMULTIPLIED_ALPHA_BIT |
 			XR_COMPOSITION_LAYER_BLEND_TEXTURE_SOURCE_ALPHA_BIT;
 		Quad.space = Layer.Desc.PositionType == ELayerType::FaceLocked ?
-			DeviceSpaces[HMDDeviceId].Space : GetTrackingSpace();
+			DeviceSpaces[HMDDeviceId].Space : PipelinedFrameStateRendering.TrackingSpace->Handle;
 		Quad.subImage.imageArrayIndex = 0;
 		Quad.pose = ToXrPose(Layer.Desc.Transform * PositionTransform, WorldToMeters);
 
@@ -2109,9 +2198,17 @@ void FOpenXRHMD::OnBeginSimulation_GameThread()
 	// xrBeginFrame so the game pipeline state can safely be modified after xrWaitFrame returns.
 	FPipelinedFrameState& PipelineState = GetPipelinedFrameStateForThread();
 	PipelineState.FrameState = FrameState;
-	PipelineState.TrackingSpace = GetTrackingSpace();
 	PipelineState.WorldToMetersScale = WorldToMetersScale;
 
+	if (bTrackingSpaceInvalid || !ensure(PipelineState.TrackingSpace.IsValid()))
+	{
+		// Create the tracking space we'll use until the next recenter.
+		FTransform BaseTransform(BaseOrientation, BasePosition);
+		PipelineState.TrackingSpace = MakeShared<FTrackingSpace>(TrackingSpaceType, ToXrPose(BaseTransform, WorldToMetersScale));
+		PipelineState.TrackingSpace->CreateSpace(Session);
+		bTrackingSpaceInvalid = false;
+	}
+
 	bShouldWait = false;
 
 	EnumerateViews(PipelineState);
@@ -2237,6 +2334,12 @@ bool FOpenXRHMD::OnStartGameFrame(FWorldContext& WorldContext)
 			if (SpaceChange.referenceSpaceType == TrackingSpaceType)
 			{
 				OnTrackingOriginChanged();
+
+				// Reset base orientation and position
+				// TODO: If poseValid is true we can use poseInPreviousSpace to make the old base transform valid in the new space
+				BaseOrientation = FQuat::Identity;
+				BasePosition = FVector::ZeroVector;
+				bTrackingSpaceInvalid = true;
 			}
 			break;
 		}
@@ -2378,7 +2481,7 @@ void FOpenXRHMD::OnFinishRendering_RHIThread()
 			Layer.type = XR_TYPE_COMPOSITION_LAYER_PROJECTION;
 			Layer.next = nullptr;
 			Layer.layerFlags = bProjectionLayerAlphaEnabled ? XR_COMPOSITION_LAYER_BLEND_TEXTURE_SOURCE_ALPHA_BIT : 0;
-			Layer.space = PipelinedFrameStateRHI.TrackingSpace;
+			Layer.space = PipelinedFrameStateRHI.TrackingSpace->Handle;
 			Layer.viewCount = PipelinedLayerStateRHI.ProjectionLayers.Num();
 			Layer.views = PipelinedLayerStateRHI.ProjectionLayers.GetData();
 			Headers.Add(reinterpret_cast<const XrCompositionLayerBaseHeader*>(&Layer));
@@ -2685,3 +2788,45 @@ void FOpenXRHMD::FDeviceSpace::DestroySpace()
 	}
 	Space = XR_NULL_HANDLE;
 }
+
+//---------------------------------------------------
+// OpenXR Tracking Space Implementation
+//---------------------------------------------------
+
+FOpenXRHMD::FTrackingSpace::FTrackingSpace(XrReferenceSpaceType InType)
+	: FTrackingSpace(InType, ToXrPose(FTransform::Identity))
+{
+}
+
+FOpenXRHMD::FTrackingSpace::FTrackingSpace(XrReferenceSpaceType InType, XrPosef InBasePose)
+	: Type(InType)
+	, Handle(XR_NULL_HANDLE)
+	, BasePose(InBasePose)
+{
+}
+
+FOpenXRHMD::FTrackingSpace::~FTrackingSpace()
+{
+	DestroySpace();
+}
+
+bool FOpenXRHMD::FTrackingSpace::CreateSpace(XrSession InSession)
+{
+	DestroySpace();
+
+	XrReferenceSpaceCreateInfo SpaceInfo;
+	SpaceInfo.type = XR_TYPE_REFERENCE_SPACE_CREATE_INFO;
+	SpaceInfo.next = nullptr;
+	SpaceInfo.referenceSpaceType = Type;
+	SpaceInfo.poseInReferenceSpace = BasePose;
+	return XR_ENSURE(xrCreateReferenceSpace(InSession, &SpaceInfo, &Handle));
+}
+
+void FOpenXRHMD::FTrackingSpace::DestroySpace()
+{
+	if (Handle)
+	{
+		XR_ENSURE(xrDestroySpace(Handle));
+	}
+	Handle = XR_NULL_HANDLE;
+}
diff --git a/Engine/Plugins/Runtime/OpenXR/Source/OpenXRHMD/Private/OpenXRHMD.h b/Engine/Plugins/Runtime/OpenXR/Source/OpenXRHMD/Private/OpenXRHMD.h
--- a/Engine/Plugins/Runtime/OpenXR/Source/OpenXRHMD/Private/OpenXRHMD.h
+++ b/Engine/Plugins/Runtime/OpenXR/Source/OpenXRHMD/Private/OpenXRHMD.h
@@ -50,6 +50,21 @@ public:
 		XrPath Path;
 	};
 
+	class FTrackingSpace
+	{
+	public:
+		FTrackingSpace(XrReferenceSpaceType InType);
+		FTrackingSpace(XrReferenceSpaceType InType, XrPosef InBasePose);
+		~FTrackingSpace();
+
+		bool CreateSpace(XrSession InSession);
+		void DestroySpace();
+
+		XrReferenceSpaceType Type;
+		XrSpace Handle;
+		XrPosef BasePose;
+	};
+
 	// The game and render threads each have a separate copy of these structures so that they don't stomp on each other or cause tearing
 	// when the game thread progresses to the next frame while the render thread is still working on the previous frame.
 	struct FPipelinedFrameState
@@ -58,7 +73,7 @@ public:
 		XrViewState ViewState{XR_TYPE_VIEW_STATE};
 		TArray<XrView> Views;
 		TArray<XrSpaceLocation> DeviceLocations;
-		XrSpace TrackingSpace = XR_NULL_HANDLE;
+		TSharedPtr<FTrackingSpace> TrackingSpace;
 		float WorldToMetersScale = 100.0f;
 
 		TArray<XrViewConfigurationView> ViewConfigs;
@@ -129,25 +144,33 @@ public:
 	virtual float GetInterpupillaryDistance() const override;
 	virtual bool GetRelativeEyePose(int32 InDeviceId, int32 InViewIndex, FQuat& OutOrientation, FVector& OutPosition) override;
 
-	virtual void ResetOrientationAndPosition(float yaw = 0.f) override;
+	virtual void ResetOrientationAndPosition(float Yaw = 0.f) override;
 	virtual void ResetOrientation(float Yaw = 0.f) override;
 	virtual void ResetPosition() override;
+	virtual void Recenter(EOrientPositionSelector::Type Selector, float Yaw = 0.f);
 
 	virtual bool GetIsTracked(int32 DeviceId);
 	virtual bool GetCurrentPose(int32 DeviceId, FQuat& CurrentOrientation, FVector& CurrentPosition) override;
 	virtual bool GetPoseForTime(int32 DeviceId, FTimespan Timespan, bool& OutTimeWasUsed, FQuat& CurrentOrientation, FVector& CurrentPosition, bool& bProvidedLinearVelocity, FVector& LinearVelocity, bool& bProvidedAngularVelocity, FVector& AngularVelocityRadPerSec, bool& bProvidedLinearAcceleration, FVector& LinearAcceleration, float WorldToMetersScale);
 	virtual bool GetCurrentInteractionProfile(const EControllerHand Hand, FString& InteractionProfile) override;
-	virtual void SetBaseRotation(const FRotator& BaseRot) override;
+	
+	virtual void SetBaseRotation(const FRotator& InBaseRotation) override;
 	virtual FRotator GetBaseRotation() const override;
 
-	virtual void SetBaseOrientation(const FQuat& BaseOrient) override;
+	virtual void SetBaseOrientation(const FQuat& InBaseOrientation) override;
 	virtual FQuat GetBaseOrientation() const override;
 
+	virtual void SetBasePosition(const FVector& InBasePosition) override;
+	virtual FVector GetBasePosition() const override;
+
 	virtual void SetTrackingOrigin(EHMDTrackingOrigin::Type NewOrigin) override
 	{
-		TrackingSpaceType = (NewOrigin == EHMDTrackingOrigin::Eye || StageSpace == XR_NULL_HANDLE) ? XR_REFERENCE_SPACE_TYPE_LOCAL : XR_REFERENCE_SPACE_TYPE_STAGE;
+		if (!bUseCustomReferenceSpace)
+		{
+			TrackingSpaceType = (NewOrigin == EHMDTrackingOrigin::Eye || StageSpace == XR_NULL_HANDLE) ? XR_REFERENCE_SPACE_TYPE_LOCAL : XR_REFERENCE_SPACE_TYPE_STAGE;
+			bTrackingSpaceInvalid = true;
+		}
 	}
-
 	virtual EHMDTrackingOrigin::Type GetTrackingOrigin() const override
 	{
 		return (TrackingSpaceType == XR_REFERENCE_SPACE_TYPE_LOCAL) ? EHMDTrackingOrigin::Eye : EHMDTrackingOrigin::Stage;
@@ -300,16 +323,8 @@ public:
 	OPENXRHMD_API XrInstance GetInstance() { return Instance; }
 	OPENXRHMD_API XrSystemId GetSystem() { return System; }
 	OPENXRHMD_API XrSession GetSession() { return Session; }
-	OPENXRHMD_API XrSpace GetTrackingSpace()
-	{
-		if (CustomSpace != XR_NULL_HANDLE)
-		{
-			return CustomSpace;
-		}
-		
-		return (TrackingSpaceType == XR_REFERENCE_SPACE_TYPE_STAGE) ? StageSpace : LocalSpace;
-	}
 	OPENXRHMD_API XrTime GetDisplayTime() const;
+	OPENXRHMD_API XrSpace GetTrackingSpace() const;
 	OPENXRHMD_API TArray<IOpenXRExtensionPlugin*>& GetExtensionPlugins() { return ExtensionPlugins; }
 
 private:
@@ -367,4 +382,9 @@ private:
 
 	TArray<FHMDViewMesh>	HiddenAreaMeshes;
 	TArray<FHMDViewMesh>	VisibleAreaMeshes;
+
+	bool					bTrackingSpaceInvalid;
+	bool					bUseCustomReferenceSpace;
+	FQuat					BaseOrientation;
+	FVector					BasePosition;
 };
