OpenXR: Properly pipeline frame state from Game Thread to Render thread

When sessions would start, there was a chance that OnBeginRendering_RenderThread could proceed without a successful call to xrWaitFrame in OnBeginSimulation_GameThread.
Additionally, the RT was accessing the GT copy of frame state directly, instead of a snapshot.
This lead to bugs like async access to frame state, invalid prediction times, or out-of-sequence frame calls (xrBeginFrame w/o preceding xrWaitFrame).

We can properly pipeline the frame state, and also add a bit about whether xrWaitFrame was completed.

diff --git a/Engine/Plugins/Runtime/OpenXR/Source/OpenXRHMD/Private/OpenXRHMD.cpp b/Engine/Plugins/Runtime/OpenXR/Source/OpenXRHMD/Private/OpenXRHMD.cpp
--- a/Engine/Plugins/Runtime/OpenXR/Source/OpenXRHMD/Private/OpenXRHMD.cpp
+++ b/Engine/Plugins/Runtime/OpenXR/Source/OpenXRHMD/Private/OpenXRHMD.cpp
@@ -495,7 +495,7 @@ bool FOpenXRHMD::GetCurrentPose(int32 DeviceId, FQuat& CurrentOrientation, FVect
 
 bool FOpenXRHMD::GetPoseForTime(int32 DeviceId, FTimespan Timespan, bool& OutTimeWasUsed, FQuat& Orientation, FVector& Position, bool& bProvidedLinearVelocity, FVector& LinearVelocity, bool& bProvidedAngularVelocity, FVector& AngularVelocityRadPerSec, bool& bProvidedLinearAcceleration, FVector& LinearAcceleration, float InWorldToMetersScale)
 {
-	FPipelinedFrameState& PipelineState = GetPipelinedFrameStateForThread();
+	const FPipelinedFrameState& PipelineState = GetPipelinedFrameStateForThread();
 
 	FReadScopeLock DeviceLock(DeviceMutex);
 	if (!DeviceSpaces.IsValidIndex(DeviceId))
@@ -579,10 +579,11 @@ void FOpenXRHMD::ResetPosition()
 
 void FOpenXRHMD::Recenter(EOrientPositionSelector::Type Selector, float Yaw)
 {
-	FPipelinedFrameState& PipelineState = GetPipelinedFrameStateForThread();
-	XrTime TargetTime = PipelineState.FrameState.predictedDisplayTime;
+	const FPipelinedFrameState& PipelineState = GetPipelinedFrameStateForThread();
+	const XrTime TargetTime = PipelineState.FrameState.predictedDisplayTime;
+	check(PipelineState.bXrFrameStateUpdated);
 
-	FDeviceSpace& DeviceSpace = DeviceSpaces[HMDDeviceId];
+	const FDeviceSpace& DeviceSpace = DeviceSpaces[HMDDeviceId];
 	XrSpaceLocation DeviceLocation = { XR_TYPE_SPACE_LOCATION, nullptr };
 
 	XrSpace BaseSpace = TrackingSpaceType == XR_REFERENCE_SPACE_TYPE_STAGE ? StageSpace : LocalSpace;
@@ -1209,8 +1210,8 @@ void FOpenXRHMD::UpdateDeviceLocations(bool bUpdateOpenXRExtensionPlugins)
 
 	FPipelinedFrameState& PipelineState = GetPipelinedFrameStateForThread();
 
-	// Only update the device locations if the frame state has been predicted
-	if (PipelineState.FrameState.predictedDisplayTime > 0)
+	// Only update the device locations if the frame state has been predicted, which is dependent on WaitFrame success
+	if (PipelineState.bXrFrameStateUpdated)
 	{
 		FReadScopeLock Lock(DeviceMutex);
 		PipelineState.DeviceLocations.SetNumZeroed(DeviceSpaces.Num());
@@ -1293,7 +1294,7 @@ void FOpenXRHMD::EnumerateViews(FPipelinedFrameState& PipelineState)
 		PipelineState.PluginViews.Add(nullptr);
 		PipelineState.ViewConfigs.Add(View);
 	}
-	XR_ENSURE(xrEnumerateViewConfigurationViews(Instance, System, SelectedViewConfigurationType, ViewConfigCount, &ViewConfigCount, PipelinedFrameStateGame.ViewConfigs.GetData()));
+	XR_ENSURE(xrEnumerateViewConfigurationViews(Instance, System, SelectedViewConfigurationType, ViewConfigCount, &ViewConfigCount, PipelineState.ViewConfigs.GetData()));
 
 	for (IOpenXRExtensionPlugin* Module : ExtensionPlugins)
 	{
@@ -1310,6 +1311,8 @@ void FOpenXRHMD::EnumerateViews(FPipelinedFrameState& PipelineState)
 	{
 		LocateViews(PipelineState, true);
 
+		check(PipelineState.bXrFrameStateUpdated);
+
 		FReadScopeLock DeviceLock(DeviceMutex);
 		for (IOpenXRExtensionPlugin* Module : ExtensionPlugins)
 		{
@@ -1734,6 +1737,7 @@ XrPath FOpenXRHMD::GetTrackedDevicePath(const int32 DeviceId)
 XrTime FOpenXRHMD::GetDisplayTime() const
 {
 	const FPipelinedFrameState& PipelineState = GetPipelinedFrameStateForThread();
+	check(PipelineState.bXrFrameStateUpdated);
 	return PipelineState.FrameState.predictedDisplayTime;
 }
 
@@ -1934,9 +1938,6 @@ void FOpenXRHMD::OnBeginRendering_RenderThread(FRHICommandListImmediate& RHICmdL
 	ensure(IsInRenderingThread());
 	FReadScopeLock DeviceLock(DeviceMutex);
 
-	UE_CLOG(PipelinedFrameStateRendering.FrameState.predictedDisplayTime >= PipelinedFrameStateGame.FrameState.predictedDisplayTime,
-		LogHMD, VeryVerbose, TEXT("Predicted display time went backwards from %lld to %lld"), PipelinedFrameStateRendering.FrameState.predictedDisplayTime, PipelinedFrameStateGame.FrameState.predictedDisplayTime);
-
 	for (IOpenXRExtensionPlugin* Module : ExtensionPlugins)
 	{
 		Module->OnBeginRendering_RenderThread(Session);
@@ -1945,8 +1946,6 @@ void FOpenXRHMD::OnBeginRendering_RenderThread(FRHICommandListImmediate& RHICmdL
 	const float WorldToMeters = GetWorldToMetersScale();
 	const FTransform InvTrackingToWorld = GetTrackingToWorldTransform().Inverse();
 
-	PipelinedFrameStateRendering = PipelinedFrameStateGame;
-
 	for (int32 ViewIndex = 0; ViewIndex < ViewFamily.Views.Num(); ViewIndex++)
 	{
 		if (ViewFamily.Views[ViewIndex]->StereoPass == EStereoscopicPass::eSSP_FULL)
@@ -2073,7 +2072,8 @@ void FOpenXRHMD::OnBeginRendering_RenderThread(FRHICommandListImmediate& RHICmdL
 	}
 #endif
 
-	if (bIsRunning)
+	// Guard prediction-dependent calls from being invoked (LocateViews, BeginFrame, etc)
+	if (bIsRunning && PipelinedFrameStateRendering.bXrFrameStateUpdated)
 	{
 		// Locate the views we will actually be rendering for.
 		// This is required to support late-updating the field-of-view.
@@ -2103,7 +2103,7 @@ void FOpenXRHMD::OnBeginRendering_RenderThread(FRHICommandListImmediate& RHICmdL
 
 void FOpenXRHMD::LocateViews(FPipelinedFrameState& PipelineState, bool ResizeViewsArray)
 {
-	check(PipelineState.FrameState.predictedDisplayTime);
+	check(PipelineState.bXrFrameStateUpdated);
 	FReadScopeLock DeviceLock(DeviceMutex);
 
 	uint32_t ViewCount = 0;
@@ -2166,15 +2166,32 @@ void FOpenXRHMD::OnBeginRendering_GameThread()
 	// can wait for the next frame in the next tick. Without this signal it's possible that two ticks
 	// happen before the next frame is actually rendered.
 	bShouldWait = true;
+
+	ENQUEUE_RENDER_COMMAND(TransferFrameStateToRenderingThread)(
+		[this, GameFrameState = PipelinedFrameStateGame](FRHICommandListImmediate& RHICmdList) mutable
+		{
+			UE_CLOG(PipelinedFrameStateRendering.FrameState.predictedDisplayTime >= GameFrameState.FrameState.predictedDisplayTime,
+				LogHMD, VeryVerbose, TEXT("Predicted display time went backwards from %lld to %lld"), PipelinedFrameStateRendering.FrameState.predictedDisplayTime, GameFrameState.FrameState.predictedDisplayTime);
+
+			PipelinedFrameStateRendering = GameFrameState;
+		});
 }
 
 void FOpenXRHMD::OnBeginSimulation_GameThread()
 {
 	FReadScopeLock Lock(SessionHandleMutex);
 
-	if (!bIsReady || !bIsRunning || !bShouldWait)
+	if (!bShouldWait)
+	{
+		return;
+	}
+
+	FPipelinedFrameState& PipelineState = GetPipelinedFrameStateForThread();
+	PipelineState.bXrFrameStateUpdated = false;
+	PipelineState.FrameState = { XR_TYPE_FRAME_STATE };
+
+	if (!bIsReady || !bIsRunning)
 	{
-		// @todo: Sleep here?
 		return;
 	}
 
@@ -2196,7 +2213,8 @@ void FOpenXRHMD::OnBeginSimulation_GameThread()
 	// The pipeline state on the game thread can only be safely modified after xrWaitFrame which will be unblocked by
 	// the runtime when xrBeginFrame is called. The rendering thread will clone the game pipeline state before calling
 	// xrBeginFrame so the game pipeline state can safely be modified after xrWaitFrame returns.
-	FPipelinedFrameState& PipelineState = GetPipelinedFrameStateForThread();
+
+	PipelineState.bXrFrameStateUpdated = true;
 	PipelineState.FrameState = FrameState;
 	PipelineState.WorldToMetersScale = WorldToMetersScale;
 
diff --git a/Engine/Plugins/Runtime/OpenXR/Source/OpenXRHMD/Private/OpenXRHMD.h b/Engine/Plugins/Runtime/OpenXR/Source/OpenXRHMD/Private/OpenXRHMD.h
--- a/Engine/Plugins/Runtime/OpenXR/Source/OpenXRHMD/Private/OpenXRHMD.h
+++ b/Engine/Plugins/Runtime/OpenXR/Source/OpenXRHMD/Private/OpenXRHMD.h
@@ -78,6 +78,8 @@ public:
 
 		TArray<XrViewConfigurationView> ViewConfigs;
 		TArray<class IOpenXRExtensionPlugin*> PluginViews;
+
+		bool bXrFrameStateUpdated = false;
 	};
 
 	struct FPipelinedLayerState
