Fix non-unity
#rnx

diff --git a/Engine/Plugins/Runtime/OpenXR/Source/OpenXRInput/Private/OpenXRInput.h b/Engine/Plugins/Runtime/OpenXR/Source/OpenXRInput/Private/OpenXRInput.h
--- a/Engine/Plugins/Runtime/OpenXR/Source/OpenXRInput/Private/OpenXRInput.h
+++ b/Engine/Plugins/Runtime/OpenXR/Source/OpenXRInput/Private/OpenXRInput.h
@@ -1,13 +1,13 @@
 // Copyright Epic Games, Inc. All Rights Reserved.
 
 #pragma once
-#include "CoreMinimal.h"
 #include "OpenXRCore.h"
 #include "GenericPlatform/IInputInterface.h"
 #include "XRMotionControllerBase.h"
 #include "IOpenXRInputPlugin.h"
 #include "IInputDevice.h"
 #include "IHapticDevice.h"
+#include "UObject/ObjectPtr.h"
 
 class FOpenXRHMD;
 class UInputAction;
