Fix missing undef of LOCTEXT_NAMESPACE in OpenXREditorModule

diff --git a/Engine/Plugins/Runtime/OpenXR/Source/OpenXREditor/Private/OpenXREditorModule.cpp b/Engine/Plugins/Runtime/OpenXR/Source/OpenXREditor/Private/OpenXREditorModule.cpp
--- a/Engine/Plugins/Runtime/OpenXR/Source/OpenXREditor/Private/OpenXREditorModule.cpp
+++ b/Engine/Plugins/Runtime/OpenXR/Source/OpenXREditor/Private/OpenXREditorModule.cpp
@@ -40,3 +40,5 @@ void FOpenXREditorModule::ShutdownModule()
 }
 
 IMPLEMENT_MODULE(FOpenXREditorModule, OpenXREditor);
+
+#undef LOCTEXT_NAMESPACE
