Don't re-use actions across action sets.

If an UInputAction is part of multiple UInputMappingContext, then we need to create multiple XrAction handles. These multiple handles are still represented by a single UInputAction for blueprints and C++.

Each XrAction gets its own set of suggested bindings and can be rebound by the player separately. This allows the developer to change the bindings for UInputAction at runtime by changing which UInputMappingContext is applied.

This is an integral part of Enhanced Input and given the restrictive nature of OpenXR it will be an important feature for XR devs as well.

diff --git a/Engine/Plugins/Runtime/OpenXR/Source/OpenXRInput/Private/OpenXRInput.cpp b/Engine/Plugins/Runtime/OpenXR/Source/OpenXRInput/Private/OpenXRInput.cpp
--- a/Engine/Plugins/Runtime/OpenXR/Source/OpenXRInput/Private/OpenXRInput.cpp
+++ b/Engine/Plugins/Runtime/OpenXR/Source/OpenXRInput/Private/OpenXRInput.cpp
@@ -558,12 +558,19 @@ void FOpenXRInputPlugin::FOpenXRInput::BuildEnhancedActions(const FSoftObjectPat
 	for (const TPair<TObjectPtr<UInputMappingContext>, int32> MappingContext : MappableInputConfig->GetMappingContexts())
 	{
 		FOpenXRActionSet ActionSet(Instance, MappingContext.Key->GetFName(), MappingContext.Key->ContextDescription.ToString(), MappingContext.Value, MappingContext.Key);
+		TMap<FName, int32> ActionMap;
+
 		for (const FEnhancedActionKeyMapping& Mapping : MappingContext.Key->GetMappings())
 		{
-			// Try to find an existing action
+			if (!Mapping.Action)
+			{
+				continue;
+			}
+
+			// Try to find an existing action within the current action set
 			FName ActionName = Mapping.Action->GetFName();
-			FOpenXRAction* Action = EnhancedActions.FindByPredicate([ActionName](const FOpenXRAction& Action) { return Action.Name == ActionName; });
-			if (!Action)
+			int32& ActionIndex = ActionMap.FindOrAdd(ActionName, INDEX_NONE);
+			if (ActionIndex == INDEX_NONE)
 			{
 				// No action found, create a new one
 				FString LocalizedName = Mapping.Action->ActionDescription.ToString();
@@ -573,12 +580,11 @@ void FOpenXRInputPlugin::FOpenXRInput::BuildEnhancedActions(const FSoftObjectPat
 					continue;
 				}
 
-				FOpenXRAction NewAction(ActionSet.Handle, ActionType, ActionName, LocalizedName, SubactionPaths, Mapping.Action);
-				int32 Index = EnhancedActions.Emplace(MoveTemp(NewAction));
-				Action = &EnhancedActions[Index];
+				// Create the action and write the index to the reference in the actions map
+				ActionIndex = EnhancedActions.Emplace(ActionSet.Handle, ActionType, ActionName, LocalizedName, SubactionPaths, Mapping.Action);
 			}
 
-			SuggestBindingForKey(Instance, *Action, Mapping.Key, Mapping.Modifiers, Mapping.Triggers);
+			SuggestBindingForKey(Instance, EnhancedActions[ActionIndex], Mapping.Key, Mapping.Modifiers, Mapping.Triggers);
 		}
 		ActionSets.Emplace(MoveTemp(ActionSet));
 	}
