Activation status of actions and axes need to be track separately

diff --git a/Engine/Plugins/Runtime/OpenXR/Source/OpenXRInput/Private/OpenXRInput.cpp b/Engine/Plugins/Runtime/OpenXR/Source/OpenXRInput/Private/OpenXRInput.cpp
--- a/Engine/Plugins/Runtime/OpenXR/Source/OpenXRInput/Private/OpenXRInput.cpp
+++ b/Engine/Plugins/Runtime/OpenXR/Source/OpenXRInput/Private/OpenXRInput.cpp
@@ -886,7 +886,7 @@ void FOpenXRInputPlugin::FOpenXRInput::SendControllerEvents()
 		Profile.next = nullptr;
 		XR_ENSURE(xrGetCurrentInteractionProfile(Session, Subaction, &Profile));
 
-		TSet<FName> ActivatedKeys;
+		TSet<FName> ActivatedActions, ActivatedAxes;
 		TPair<XrPath, XrPath> Key(Profile.interactionProfile, Subaction);
 		for (FOpenXRAction& Action : LegacyActions)
 		{
@@ -898,7 +898,7 @@ void FOpenXRInputPlugin::FOpenXRInput::SendControllerEvents()
 
 			// Find the action key and check if it has already been fired this frame.
 			FName* ActionKey = Action.KeyMap.Find(Key);
-			if (!ActionKey || ActivatedKeys.Contains(*ActionKey))
+			if (!ActionKey)
 			{
 				continue;
 			}
@@ -906,6 +906,7 @@ void FOpenXRInputPlugin::FOpenXRInput::SendControllerEvents()
 			switch (Action.Type)
 			{
 			case XR_ACTION_TYPE_BOOLEAN_INPUT:
+			if (!ActivatedActions.Contains(*ActionKey))
 			{
 				XrActionStateBoolean State;
 				State.type = XR_TYPE_ACTION_STATE_BOOLEAN;
@@ -913,7 +914,7 @@ void FOpenXRInputPlugin::FOpenXRInput::SendControllerEvents()
 				XrResult Result = xrGetActionStateBoolean(Session, &GetInfo, &State);
 				if (XR_SUCCEEDED(Result) && State.changedSinceLastSync)
 				{
-					ActivatedKeys.Add(*ActionKey);
+					ActivatedActions.Add(*ActionKey);
 					if (State.isActive && State.currentState)
 					{
 						MessageHandler->OnControllerButtonPressed(*ActionKey, 0, /*IsRepeat =*/false);
@@ -932,6 +933,7 @@ void FOpenXRInputPlugin::FOpenXRInput::SendControllerEvents()
 			}
 			break;
 			case XR_ACTION_TYPE_FLOAT_INPUT:
+			if (!ActivatedAxes.Contains(*ActionKey))
 			{
 				XrActionStateFloat State;
 				State.type = XR_TYPE_ACTION_STATE_FLOAT;
@@ -939,7 +941,7 @@ void FOpenXRInputPlugin::FOpenXRInput::SendControllerEvents()
 				XrResult Result = xrGetActionStateFloat(Session, &GetInfo, &State);
 				if (XR_SUCCEEDED(Result) && State.changedSinceLastSync)
 				{
-					ActivatedKeys.Add(*ActionKey);
+					ActivatedAxes.Add(*ActionKey);
 					if (State.isActive)
 					{
 						MessageHandler->OnControllerAnalog(*ActionKey, 0, State.currentState);
