OpenXR: Register with FCoreDelegates::VRHeadsetRecenter

Now that FOpenXRHMD supports ResetOrientationAndPosition correctly, we can register with the FCoreDelegate for those clients.

diff --git a/Engine/Plugins/Runtime/OpenXR/Source/OpenXRHMD/Private/OpenXRHMD.cpp b/Engine/Plugins/Runtime/OpenXR/Source/OpenXRHMD/Private/OpenXRHMD.cpp
--- a/Engine/Plugins/Runtime/OpenXR/Source/OpenXRHMD/Private/OpenXRHMD.cpp
+++ b/Engine/Plugins/Runtime/OpenXR/Source/OpenXRHMD/Private/OpenXRHMD.cpp
@@ -19,6 +19,7 @@
 #include "PostProcess/PostProcessHMD.h"
 #include "GameFramework/WorldSettings.h"
 #include "Misc/CString.h"
+#include "Misc/CoreDelegates.h"
 #include "ClearQuad.h"
 #include "XRThreadUtils.h"
 #include "RenderUtils.h"
@@ -561,6 +562,11 @@ bool FOpenXRHMD::IsChromaAbCorrectionEnabled() const
 	return false;
 }
 
+void FOpenXRHMD::VRHeadsetRecenterDelegate()
+{
+	Recenter(EOrientPositionSelector::OrientationAndPosition, 0.f);
+}
+
 void FOpenXRHMD::ResetOrientationAndPosition(float Yaw)
 {
 	Recenter(EOrientPositionSelector::OrientationAndPosition, Yaw);
@@ -1597,6 +1603,8 @@ bool FOpenXRHMD::OnStereoStartup()
 		}
 	}
 
+	FCoreDelegates::VRHeadsetRecenter.AddRaw(this, &FOpenXRHMD::VRHeadsetRecenterDelegate);
+
 	return true;
 }
 
@@ -1622,6 +1630,9 @@ bool FOpenXRHMD::OnStereoTeardown()
 	{
 		XR_ENSURE(Result);
 	}
+
+	FCoreDelegates::VRHeadsetRecenter.RemoveAll(this);
+
 	return true;
 }
 
diff --git a/Engine/Plugins/Runtime/OpenXR/Source/OpenXRHMD/Private/OpenXRHMD.h b/Engine/Plugins/Runtime/OpenXR/Source/OpenXRHMD/Private/OpenXRHMD.h
--- a/Engine/Plugins/Runtime/OpenXR/Source/OpenXRHMD/Private/OpenXRHMD.h
+++ b/Engine/Plugins/Runtime/OpenXR/Source/OpenXRHMD/Private/OpenXRHMD.h
@@ -228,6 +228,9 @@ protected:
 	
 	void CopyTexture_RenderThread(FRHICommandListImmediate& RHICmdList, FRHITexture2D* SrcTexture, FIntRect SrcRect, const FXRSwapChainPtr& DstSwapChain, FIntRect DstRect, bool bClearBlack, bool bNoAlpha) const;
 
+	// Used with FCoreDelegates
+	void VRHeadsetRecenterDelegate();
+
 public:
 	/** IXRTrackingSystem interface */
 	virtual bool DoesSupportLateProjectionUpdate() const override { return true; }
