OpenXR Display mapping for devices that need it

diff --git a/Engine/Plugins/Runtime/OpenXR/Source/OpenXRHMD/Private/OpenXRHMD.cpp b/Engine/Plugins/Runtime/OpenXR/Source/OpenXRHMD/Private/OpenXRHMD.cpp
--- a/Engine/Plugins/Runtime/OpenXR/Source/OpenXRHMD/Private/OpenXRHMD.cpp
+++ b/Engine/Plugins/Runtime/OpenXR/Source/OpenXRHMD/Private/OpenXRHMD.cpp
@@ -32,6 +32,7 @@
 #include "ScenePrivate.h"
 #include "GeneralProjectSettings.h"
 #include "Epic_openxr.h"
+#include "HDRHelper.h"
 
 #if WITH_EDITOR
 #include "Editor/UnrealEd/Classes/Editor/EditorEngine.h"
@@ -2619,6 +2620,148 @@ FIntRect FOpenXRHMD::GetFullFlatEyeRect_RenderThread(FTexture2DRHIRef EyeTexture
 	return FIntRect(EyeTexture->GetSizeX() * SrcNormRectMin.X, EyeTexture->GetSizeY() * SrcNormRectMin.Y, EyeTexture->GetSizeX() * SrcNormRectMax.X, EyeTexture->GetSizeY() * SrcNormRectMax.Y);
 }
 
+class FDisplayMappingPS : public FGlobalShader
+{
+	DECLARE_SHADER_TYPE(FDisplayMappingPS, Global);
+public:
+
+	static bool ShouldCompilePermutation(const FGlobalShaderPermutationParameters& Parameters)
+	{
+		return true;
+	}
+
+	static void ModifyCompilationEnvironment(const FGlobalShaderPermutationParameters& Parameters, FShaderCompilerEnvironment& OutEnvironment)
+	{
+		FGlobalShader::ModifyCompilationEnvironment(Parameters, OutEnvironment);
+		OutEnvironment.SetDefine(TEXT("DISPLAY_MAPPING_PS"), 1);
+	}
+
+
+	FDisplayMappingPS(const ShaderMetaType::CompiledShaderInitializerType& Initializer) :
+		FGlobalShader(Initializer)
+	{
+		OutputDevice.Bind(Initializer.ParameterMap, TEXT("OutputDevice"));
+		OutputGamut.Bind(Initializer.ParameterMap, TEXT("OutputGamut"));
+		SceneTexture.Bind(Initializer.ParameterMap, TEXT("SceneTexture"));
+		SceneSampler.Bind(Initializer.ParameterMap, TEXT("SceneSampler"));
+		TextureToOutputGamutMatrix.Bind(Initializer.ParameterMap, TEXT("TextureToOutputGamutMatrix"));
+	}
+	FDisplayMappingPS() = default;
+
+	static FMatrix44f GamutToXYZMatrix(EDisplayColorGamut ColorGamut)
+	{
+		static const FMatrix44f sRGB_2_XYZ_MAT(
+			FVector3f(	0.4124564, 0.3575761, 0.1804375),
+			FVector3f(	0.2126729, 0.7151522, 0.0721750),
+			FVector3f(	0.0193339, 0.1191920, 0.9503041),
+			FVector3f(	0        ,         0,         0)
+		);
+
+		static const FMatrix44f Rec2020_2_XYZ_MAT(
+			FVector3f(	0.6369736, 0.1446172, 0.1688585),
+			FVector3f(	0.2627066, 0.6779996, 0.0592938),
+			FVector3f(	0.0000000, 0.0280728, 1.0608437),
+			FVector3f(	0        ,         0,         0)
+		);
+
+		static const FMatrix44f P3D65_2_XYZ_MAT(
+			FVector3f(	0.4865906, 0.2656683, 0.1981905),
+			FVector3f(	0.2289838, 0.6917402, 0.0792762),
+			FVector3f(	0.0000000, 0.0451135, 1.0438031),
+			FVector3f(	0        ,         0,         0)
+		);
+		switch (ColorGamut)
+		{
+		case EDisplayColorGamut::sRGB_D65: return sRGB_2_XYZ_MAT;
+		case EDisplayColorGamut::Rec2020_D65: return Rec2020_2_XYZ_MAT;
+		case EDisplayColorGamut::DCIP3_D65: return P3D65_2_XYZ_MAT;
+		default:
+			checkNoEntry();
+			return FMatrix44f::Identity;
+		}
+
+	}
+
+	static FMatrix44f XYZToGamutMatrix(EDisplayColorGamut ColorGamut)
+	{
+		static const FMatrix44f XYZ_2_sRGB_MAT(
+			FVector3f(	 3.2409699419, -1.5373831776, -0.4986107603),
+			FVector3f(	-0.9692436363,  1.8759675015,  0.0415550574),
+			FVector3f(	 0.0556300797, -0.2039769589,  1.0569715142),
+			FVector3f(	 0           ,             0,             0)
+		);
+
+		static const FMatrix44f XYZ_2_Rec2020_MAT(
+			FVector3f(1.7166084, -0.3556621, -0.2533601),
+			FVector3f(-0.6666829, 1.6164776, 0.0157685),
+			FVector3f(0.0176422, -0.0427763, 0.94222867),
+			FVector3f(0, 0, 0)
+		);
+
+		static const FMatrix44f XYZ_2_P3D65_MAT(
+			FVector3f(	 2.4933963, -0.9313459, -0.4026945),
+			FVector3f(	-0.8294868,  1.7626597,  0.0236246),
+			FVector3f(	 0.0358507, -0.0761827,  0.9570140),
+			FVector3f(	 0        ,         0,         0 )
+		);
+
+		switch (ColorGamut)
+		{
+		case EDisplayColorGamut::sRGB_D65: return XYZ_2_sRGB_MAT;
+		case EDisplayColorGamut::Rec2020_D65: return XYZ_2_Rec2020_MAT;
+		case EDisplayColorGamut::DCIP3_D65: return XYZ_2_P3D65_MAT;
+		default:
+			checkNoEntry();
+			return FMatrix44f::Identity;
+		}
+
+	}
+
+	void SetParameters(FRHICommandList& RHICmdList, EDisplayOutputFormat DisplayOutputFormat, EDisplayColorGamut DisplayColorGamut, EDisplayColorGamut TextureColorGamut, FRHITexture* SceneTextureRHI, bool bSameSize)
+	{
+		int32 OutputDeviceValue = (int32)DisplayOutputFormat;
+		int32 OutputGamutValue = (int32)DisplayColorGamut;
+
+		SetShaderValue(RHICmdList, RHICmdList.GetBoundPixelShader(), OutputDevice, OutputDeviceValue);
+		SetShaderValue(RHICmdList, RHICmdList.GetBoundPixelShader(), OutputGamut, OutputGamutValue);
+
+		const FMatrix44f TextureGamutMatrixToXYZ = GamutToXYZMatrix(TextureColorGamut);
+		const FMatrix44f XYZToDisplayMatrix = XYZToGamutMatrix(DisplayColorGamut);
+		// note: we use mul(m,v) instead of mul(v,m) in the shaders for color conversions which is why matrix multiplication is reversed compared to what we usually do
+		const FMatrix44f CombinedMatrix = XYZToDisplayMatrix * TextureGamutMatrixToXYZ;
+
+		SetShaderValue(RHICmdList, RHICmdList.GetBoundPixelShader(), TextureToOutputGamutMatrix, CombinedMatrix);
+
+		if (bSameSize)
+		{
+			SetTextureParameter(RHICmdList, RHICmdList.GetBoundPixelShader(), SceneTexture, SceneSampler, TStaticSamplerState<SF_Point>::GetRHI(), SceneTextureRHI);
+		}
+		else
+		{
+			SetTextureParameter(RHICmdList, RHICmdList.GetBoundPixelShader(), SceneTexture, SceneSampler, TStaticSamplerState<SF_Bilinear>::GetRHI(), SceneTextureRHI);
+		}
+	}
+
+	static const TCHAR* GetSourceFilename()
+	{
+		return TEXT("/Engine/Private/CompositeUIPixelShader.usf");
+	}
+
+	static const TCHAR* GetFunctionName()
+	{
+		return TEXT("DisplayMappingPS");
+	}
+
+private:
+	LAYOUT_FIELD(FShaderParameter, OutputDevice);
+	LAYOUT_FIELD(FShaderParameter, OutputGamut);
+	LAYOUT_FIELD(FShaderParameter, TextureToOutputGamutMatrix);
+	LAYOUT_FIELD(FShaderResourceParameter, SceneTexture);
+	LAYOUT_FIELD(FShaderResourceParameter, SceneSampler);
+};
+
+IMPLEMENT_SHADER_TYPE(, FDisplayMappingPS, TEXT("/Engine/Private/CompositeUIPixelShader.usf"), TEXT("DisplayMappingPS"), SF_Pixel);
+
 void FOpenXRHMD::CopyTexture_RenderThread(FRHICommandListImmediate& RHICmdList, FRHITexture2D* SrcTexture, FIntRect SrcRect, FRHITexture2D* DstTexture, FIntRect DstRect, 
 											bool bClearBlack, bool bNoAlpha, ERenderTargetActions RTAction, ERHIAccess FinalDstAccess) const
 {
@@ -2672,7 +2815,45 @@ void FOpenXRHMD::CopyTexture_RenderThread(FRHICommandListImmediate& RHICmdList,
 		auto ShaderMap = GetGlobalShaderMap(FeatureLevel);
 
 		TShaderMapRef<FScreenVS> VertexShader(ShaderMap);
-		TShaderMapRef<FScreenPS> PixelShader(ShaderMap);
+
+		TShaderRef<FGlobalShader> PixelShader;
+		FDisplayMappingPS* DisplayMappingPS = nullptr;
+		FScreenPS* ScreenPS = nullptr;
+
+		bool bNeedsDisplayMapping = false;
+		EDisplayOutputFormat TVDisplayOutputFormat = EDisplayOutputFormat::SDR_sRGB;
+		EDisplayColorGamut HMDColorGamut = EDisplayColorGamut::sRGB_D65;
+		EDisplayColorGamut TVColorGamut = EDisplayColorGamut::sRGB_D65;
+		if (FinalDstAccess == ERHIAccess::Present && RenderBridge.IsValid())
+		{
+			EDisplayOutputFormat HMDDisplayFormat;
+			bool bHMDSupportHDR;
+			if (RenderBridge->HDRGetMetaDataForStereo(HMDDisplayFormat, HMDColorGamut, bHMDSupportHDR))
+			{
+				bool bTVSupportHDR;
+				HDRGetMetaData(TVDisplayOutputFormat, TVColorGamut, bTVSupportHDR, FVector2D(0, 0), FVector2D(0, 0), nullptr);
+				if (TVDisplayOutputFormat != HMDDisplayFormat || HMDColorGamut != TVColorGamut || bTVSupportHDR != bHMDSupportHDR)
+				{
+					// shader assumes G 2.2 for input / ST2084/sRGB for output right now
+					ensure(HMDDisplayFormat == EDisplayOutputFormat::SDR_ExplicitGammaMapping);
+					ensure(TVDisplayOutputFormat == EDisplayOutputFormat::SDR_sRGB || TVDisplayOutputFormat == EDisplayOutputFormat::HDR_ACES_1000nit_ST2084 || TVDisplayOutputFormat == EDisplayOutputFormat::HDR_ACES_2000nit_ST2084);
+					bNeedsDisplayMapping = true;
+				}
+			}
+		}
+
+		if (bNeedsDisplayMapping)
+		{
+			TShaderMapRef<FDisplayMappingPS> DisplayMappingPSRef(ShaderMap);
+			DisplayMappingPS = DisplayMappingPSRef.GetShader();
+			PixelShader = DisplayMappingPSRef;
+		}
+		else
+		{
+			TShaderMapRef<FScreenPS> ScreenPSRef(ShaderMap);
+			ScreenPS = ScreenPSRef.GetShader();
+			PixelShader = ScreenPSRef;
+		}
 
 		GraphicsPSOInit.BoundShaderState.VertexDeclarationRHI = GFilterVertexDeclaration.VertexDeclarationRHI;
 		GraphicsPSOInit.BoundShaderState.VertexShaderRHI = VertexShader.GetVertexShader();
@@ -2683,13 +2864,20 @@ void FOpenXRHMD::CopyTexture_RenderThread(FRHICommandListImmediate& RHICmdList,
 		RHICmdList.Transition(FRHITransitionInfo(SrcTexture, ERHIAccess::Unknown, ERHIAccess::SRVMask));
 
 		const bool bSameSize = DstRect.Size() == SrcRect.Size();
-		if (bSameSize)
+		if (ScreenPS)
 		{
-			PixelShader->SetParameters(RHICmdList, TStaticSamplerState<SF_Point>::GetRHI(), SrcTexture);
+			if (bSameSize)
+			{
+				ScreenPS->SetParameters(RHICmdList, TStaticSamplerState<SF_Point>::GetRHI(), SrcTexture);
+			}
+			else
+			{
+				ScreenPS->SetParameters(RHICmdList, TStaticSamplerState<SF_Bilinear>::GetRHI(), SrcTexture);
+			}
 		}
-		else
+		else if (DisplayMappingPS)
 		{
-			PixelShader->SetParameters(RHICmdList, TStaticSamplerState<SF_Bilinear>::GetRHI(), SrcTexture);
+			DisplayMappingPS->SetParameters(RHICmdList, TVDisplayOutputFormat, TVColorGamut, HMDColorGamut, SrcTexture, bSameSize);
 		}
 
 		RendererModule->DrawRectangle(
diff --git a/Engine/Shaders/Private/CompositeUIPixelShader.usf b/Engine/Shaders/Private/CompositeUIPixelShader.usf
--- a/Engine/Shaders/Private/CompositeUIPixelShader.usf
+++ b/Engine/Shaders/Private/CompositeUIPixelShader.usf
@@ -195,4 +195,31 @@ void Main(
 	const float3x3 Rec2020_2_sRGB = mul(XYZ_2_sRGB_MAT, Rec2020_2_XYZ_MAT);
 	OutColor.xyz = mul(Rec2020_2_sRGB, OutColor.xyz / PaperWhiteInNits);
 #endif
-}
\ No newline at end of file
+}
+
+#if DISPLAY_MAPPING_PS
+float4x4 TextureToOutputGamutMatrix;
+
+// Currently hard-coded to go from G2.2 to either PQ/Rec2020 or sRGB
+void DisplayMappingPS(
+	FScreenVertexOutput Input,
+	out float4 OutColor : SV_Target0
+	)
+{
+	const float4 SceneColor = Texture2DSampleLevel(SceneTexture, SceneSampler, Input.UV, 0) ;
+	OutColor = SceneColor;
+
+	float3 LinearSceneColor = pow(SceneColor.rgb, 2.2);
+	LinearSceneColor = mul((float3x3)TextureToOutputGamutMatrix, LinearSceneColor);
+	
+	if(OutputDevice == TONEMAPPER_OUTPUT_ACES1000nitST2084 || OutputDevice == TONEMAPPER_OUTPUT_ACES2000nitST2084)
+	{
+		OutColor.xyz = LinearToST2084(LinearSceneColor*300.f);
+	}
+	else
+	{
+		OutColor.xyz = LinearToSrgbBranchless(saturate(LinearSceneColor));
+	}
+}
+
+#endif
\ No newline at end of file
