Fix mobile cook. Compile FDisplayMappingPS only for SM5, FWriteToSliceGeometryOutput not defined on mobile

diff --git a/Engine/Plugins/Runtime/OpenXR/Source/OpenXRHMD/Private/OpenXRHMD.cpp b/Engine/Plugins/Runtime/OpenXR/Source/OpenXRHMD/Private/OpenXRHMD.cpp
--- a/Engine/Plugins/Runtime/OpenXR/Source/OpenXRHMD/Private/OpenXRHMD.cpp
+++ b/Engine/Plugins/Runtime/OpenXR/Source/OpenXRHMD/Private/OpenXRHMD.cpp
@@ -2627,7 +2627,7 @@ public:
 
 	static bool ShouldCompilePermutation(const FGlobalShaderPermutationParameters& Parameters)
 	{
-		return true;
+		return IsFeatureLevelSupported(Parameters.Platform, ERHIFeatureLevel::SM5);
 	}
 
 	static void ModifyCompilationEnvironment(const FGlobalShaderPermutationParameters& Parameters, FShaderCompilerEnvironment& OutEnvironment)
@@ -2842,6 +2842,8 @@ void FOpenXRHMD::CopyTexture_RenderThread(FRHICommandListImmediate& RHICmdList,
 			}
 		}
 
+		bNeedsDisplayMapping &= (GMaxRHIFeatureLevel > ERHIFeatureLevel::ES3_1);
+
 		if (bNeedsDisplayMapping)
 		{
 			TShaderMapRef<FDisplayMappingPS> DisplayMappingPSRef(ShaderMap);
