Update OpenXR to use the new version of the "OnControllerX" functions.

#rnx

diff --git a/Engine/Plugins/Runtime/OpenXR/Source/OpenXRInput/Private/OpenXRInput.cpp b/Engine/Plugins/Runtime/OpenXR/Source/OpenXRInput/Private/OpenXRInput.cpp
--- a/Engine/Plugins/Runtime/OpenXR/Source/OpenXRInput/Private/OpenXRInput.cpp
+++ b/Engine/Plugins/Runtime/OpenXR/Source/OpenXRInput/Private/OpenXRInput.cpp
@@ -13,6 +13,7 @@
 #include "InputAction.h"
 #include "InputMappingContext.h"
 #include "PlayerMappableInputConfig.h"
+#include "GenericPlatform/GenericPlatformInputDeviceMapper.h"
 
 #if WITH_EDITOR
 #include "Editor/EditorEngine.h"
@@ -912,16 +913,19 @@ void FOpenXRInputPlugin::FOpenXRInput::SendControllerEvents()
 				State.type = XR_TYPE_ACTION_STATE_BOOLEAN;
 				State.next = nullptr;
 				XrResult Result = xrGetActionStateBoolean(Session, &GetInfo, &State);
+				IPlatformInputDeviceMapper& DeviceMapper = IPlatformInputDeviceMapper::Get();
 				if (XR_SUCCEEDED(Result) && State.changedSinceLastSync)
 				{
 					ActivatedActions.Add(*ActionKey);
 					if (State.isActive && State.currentState)
 					{
-						MessageHandler->OnControllerButtonPressed(*ActionKey, 0, /*IsRepeat =*/false);
+						// TODO: Map input devices here if OpenXR would like to
+						MessageHandler->OnControllerButtonPressed(*ActionKey, DeviceMapper.GetPrimaryPlatformUser(), DeviceMapper.GetDefaultInputDevice(), /*IsRepeat =*/false);
 					}
 					else
 					{
-						MessageHandler->OnControllerButtonReleased(*ActionKey, 0, /*IsRepeat =*/false);
+						// TODO: Map input devices here if OpenXR would like to
+						MessageHandler->OnControllerButtonReleased(*ActionKey, DeviceMapper.GetPrimaryPlatformUser(), DeviceMapper.GetDefaultInputDevice(), /*IsRepeat =*/false);
 					}
 
 					FXRTimedInputActionDelegate* const Delegate = OpenXRInputNamespace::GetTimedInputActionDelegate(Action.Name);
@@ -939,16 +943,19 @@ void FOpenXRInputPlugin::FOpenXRInput::SendControllerEvents()
 				State.type = XR_TYPE_ACTION_STATE_FLOAT;
 				State.next = nullptr;
 				XrResult Result = xrGetActionStateFloat(Session, &GetInfo, &State);
+				IPlatformInputDeviceMapper& DeviceMapper = IPlatformInputDeviceMapper::Get();
 				if (XR_SUCCEEDED(Result) && State.changedSinceLastSync)
 				{
 					ActivatedAxes.Add(*ActionKey);
 					if (State.isActive)
 					{
-						MessageHandler->OnControllerAnalog(*ActionKey, 0, State.currentState);
+						// TODO: Map input devices here if OpenXR would like to
+						MessageHandler->OnControllerAnalog(*ActionKey, DeviceMapper.GetPrimaryPlatformUser(), DeviceMapper.GetDefaultInputDevice(), State.currentState);
 					}
 					else
 					{
-						MessageHandler->OnControllerAnalog(*ActionKey, 0, 0.0f);
+						// TODO: Map input devices here if OpenXR would like to
+						MessageHandler->OnControllerAnalog(*ActionKey, DeviceMapper.GetPrimaryPlatformUser(), DeviceMapper.GetDefaultInputDevice(), 0.0f);
 					}
 
 					FXRTimedInputActionDelegate* const Delegate = OpenXRInputNamespace::GetTimedInputActionDelegate(Action.Name);
