OpenXR: Set PixelDensity (but don't use it yet!)

Initial change to support setting PixelDensity in OpenXR. Also adds pipelined access to PixelDensity, and making sure app-set PixelDensity is not more than the runtime actually supports.

diff --git a/Engine/Plugins/Runtime/OpenXR/Source/OpenXRHMD/Private/OpenXRHMD.cpp b/Engine/Plugins/Runtime/OpenXR/Source/OpenXRHMD/Private/OpenXRHMD.cpp
--- a/Engine/Plugins/Runtime/OpenXR/Source/OpenXRHMD/Private/OpenXRHMD.cpp
+++ b/Engine/Plugins/Runtime/OpenXR/Source/OpenXRHMD/Private/OpenXRHMD.cpp
@@ -1115,6 +1115,14 @@ FOpenXRHMD::FOpenXRHMD(const FAutoRegister& AutoRegister, XrInstance InInstance,
 	// Enumerate the views we will be simulating with.
 	EnumerateViews(PipelinedFrameStateGame);
 
+	for (const XrViewConfigurationView& Config : PipelinedFrameStateGame.ViewConfigs)
+	{
+		const float WidthDensityMax = float(Config.maxImageRectWidth) / Config.recommendedImageRectWidth;
+		const float HeightDensitymax = float(Config.maxImageRectHeight) / Config.recommendedImageRectHeight;
+		const float PerViewPixelDensityMax = FMath::Min(WidthDensityMax, HeightDensitymax);
+		RuntimePixelDensityMax = FMath::Min(RuntimePixelDensityMax, PerViewPixelDensityMax);
+	}
+
 	// Enumerate environment blend modes and select the best one.
 	{
 		uint32 BlendModeCount;
@@ -2577,6 +2585,12 @@ bool FOpenXRHMD::HDRGetMetaDataForStereo(EDisplayOutputFormat& OutDisplayOutputF
 	return RenderBridge->HDRGetMetaDataForStereo(OutDisplayOutputFormat, OutDisplayColorGamut, OutbHDRSupported);
 }
 
+void FOpenXRHMD::SetPixelDensity(const float NewDensity)
+{
+	check(IsInGameThread());
+	PipelinedFrameStateGame.PixelDensity = FMath::Min(NewDensity, RuntimePixelDensityMax);
+}
+
 FIntPoint FOpenXRHMD::GetIdealRenderTargetSize() const
 {
 	const FPipelinedFrameState& PipelineState = GetPipelinedFrameStateForThread();
diff --git a/Engine/Plugins/Runtime/OpenXR/Source/OpenXRHMD/Private/OpenXRHMD.h b/Engine/Plugins/Runtime/OpenXR/Source/OpenXRHMD/Private/OpenXRHMD.h
--- a/Engine/Plugins/Runtime/OpenXR/Source/OpenXRHMD/Private/OpenXRHMD.h
+++ b/Engine/Plugins/Runtime/OpenXR/Source/OpenXRHMD/Private/OpenXRHMD.h
@@ -75,6 +75,7 @@ public:
 		TArray<XrSpaceLocation> DeviceLocations;
 		TSharedPtr<FTrackingSpace> TrackingSpace;
 		float WorldToMetersScale = 100.0f;
+		float PixelDensity = 1.0f;
 
 		TArray<XrViewConfigurationView> ViewConfigs;
 		TArray<class IOpenXRExtensionPlugin*> PluginViews;
@@ -246,6 +247,8 @@ public:
 	virtual bool GetHMDMonitorInfo(MonitorInfo&) override;
 	virtual void GetFieldOfView(float& OutHFOVInDegrees, float& OutVFOVInDegrees) const override;
 	virtual bool IsChromaAbCorrectionEnabled() const override;
+	//virtual float GetPixelDenity() const override;
+	virtual void SetPixelDensity(const float NewDensity) override;
 	virtual FIntPoint GetIdealRenderTargetSize() const override;
 	virtual bool GetHMDDistortionEnabled(EShadingPath ShadingPath) const override { return false; }
 	virtual FIntRect GetFullFlatEyeRect_RenderThread(FTexture2DRHIRef EyeTexture) const override;
@@ -354,6 +357,7 @@ private:
 	bool					bProjectionLayerAlphaEnabled;
 	bool					bIsStandaloneStereoOnlyDevice;
 	float					WorldToMetersScale = 100.0f;
+	float					RuntimePixelDensityMax = FHeadMountedDisplayBase::PixelDensityMax;
 
 	XrSessionState			CurrentSessionState;
 	FRWLock					SessionHandleMutex;
