Direct FTCHARToUTF8_Convert usage replaced with StringCast or FPlatformString in OpenXR.

diff --git a/Engine/Plugins/Runtime/OpenXR/Source/OpenXRHMD/Private/OpenXRHMDModule.cpp b/Engine/Plugins/Runtime/OpenXR/Source/OpenXRHMD/Private/OpenXRHMDModule.cpp
--- a/Engine/Plugins/Runtime/OpenXR/Source/OpenXRHMD/Private/OpenXRHMDModule.cpp
+++ b/Engine/Plugins/Runtime/OpenXR/Source/OpenXRHMD/Private/OpenXRHMDModule.cpp
@@ -637,9 +637,9 @@ bool FOpenXRHMDModule::InitInstance()
 	Info.type = XR_TYPE_INSTANCE_CREATE_INFO;
 	Info.next = nullptr;
 	Info.createFlags = 0;
-	FTCHARToUTF8_Convert::Convert(Info.applicationInfo.applicationName, XR_MAX_APPLICATION_NAME_SIZE, *AppName, AppName.Len() + 1);
+	FPlatformString::Convert((UTF8CHAR*)Info.applicationInfo.applicationName, XR_MAX_APPLICATION_NAME_SIZE, *AppName, AppName.Len() + 1);
 	Info.applicationInfo.applicationVersion = static_cast<uint32>(BuildSettings::GetCurrentChangelist()) | (BuildSettings::IsLicenseeVersion() ? 0x80000000 : 0);
-	FTCHARToUTF8_Convert::Convert(Info.applicationInfo.engineName, XR_MAX_ENGINE_NAME_SIZE, *EngineName, EngineName.Len() + 1);
+	FPlatformString::Convert((UTF8CHAR*)Info.applicationInfo.engineName, XR_MAX_ENGINE_NAME_SIZE, *EngineName, EngineName.Len() + 1);
 	Info.applicationInfo.engineVersion = (uint32)(FEngineVersion::Current().GetMinor() << 16 | FEngineVersion::Current().GetPatch());
 	Info.applicationInfo.apiVersion = XR_CURRENT_API_VERSION;
 
