OpenXR: Allow for mono pass from plugin driven views

Currently, we consider all plugin views to be stereo, but this isn't always the case (ex: HL2)
We infer whether the plugin view is mono or stereo. In the future, we can add an extension callback to set the field.

diff --git a/Engine/Plugins/Runtime/OpenXR/Source/OpenXRHMD/Private/OpenXRHMD.cpp b/Engine/Plugins/Runtime/OpenXR/Source/OpenXRHMD/Private/OpenXRHMD.cpp
--- a/Engine/Plugins/Runtime/OpenXR/Source/OpenXRHMD/Private/OpenXRHMD.cpp
+++ b/Engine/Plugins/Runtime/OpenXR/Source/OpenXRHMD/Private/OpenXRHMD.cpp
@@ -785,13 +785,13 @@ void FOpenXRHMD::SetFinalViewRect(FRHICommandListImmediate& RHICmdList, const in
 		DepthImage.imageRect = ColorImage.imageRect;
 	}
 
-	if (!PipelinedFrameStateRendering.PluginViews.IsValidIndex(ViewIndex))
+	if (!PipelinedFrameStateRendering.PluginViewInfos.IsValidIndex(ViewIndex))
 	{
 		// This plugin is no longer providing this view.
 		return;
 	}
 
-	if (PipelinedFrameStateRendering.PluginViews[ViewIndex])
+	if (PipelinedFrameStateRendering.PluginViewInfos[ViewIndex].bIsPluginManaged)
 	{
 		// Defer to the plugin to handle submission
 		return;
@@ -838,10 +838,9 @@ EStereoscopicPass FOpenXRHMD::GetViewPassForIndex(bool bStereoRequested, int32 V
 		return EStereoscopicPass::eSSP_FULL;
 
 	const FPipelinedFrameState& PipelineState = GetPipelinedFrameStateForThread();
-	if (PipelineState.PluginViews.IsValidIndex(ViewIndex) && PipelineState.PluginViews[ViewIndex])
+	if (PipelineState.PluginViewInfos.IsValidIndex(ViewIndex) && PipelineState.PluginViewInfos[ViewIndex].bIsPluginManaged)
 	{
-		// Views provided by a plugin should be considered a new primary pass
-		return EStereoscopicPass::eSSP_PRIMARY;
+		return PipelineState.PluginViewInfos[ViewIndex].PassType;
 	}
 
 	if (SelectedViewConfigurationType == XR_VIEW_CONFIGURATION_TYPE_PRIMARY_QUAD_VARJO)
@@ -1290,7 +1289,7 @@ void FOpenXRHMD::EnumerateViews(FPipelinedFrameState& PipelineState)
 	XR_ENSURE(xrEnumerateViewConfigurationViews(Instance, System, SelectedViewConfigurationType, 0, &ViewConfigCount, nullptr));
 	ViewFov.SetNum(ViewConfigCount);
 	PipelineState.ViewConfigs.Empty(ViewConfigCount);
-	PipelineState.PluginViews.Empty(ViewConfigCount);
+	PipelineState.PluginViewInfos.Empty(ViewConfigCount);
 	for (uint32 ViewIndex = 0; ViewIndex < ViewConfigCount; ViewIndex++)
 	{
 		XrViewConfigurationView View;
@@ -1306,7 +1305,7 @@ void FOpenXRHMD::EnumerateViews(FPipelinedFrameState& PipelineState)
 		}
 
 		// These are core views that don't have an associated plugin
-		PipelineState.PluginViews.Add(nullptr);
+		PipelineState.PluginViewInfos.AddDefaulted(1);
 		PipelineState.ViewConfigs.Add(View);
 	}
 	XR_ENSURE(xrEnumerateViewConfigurationViews(Instance, System, SelectedViewConfigurationType, ViewConfigCount, &ViewConfigCount, PipelineState.ViewConfigs.GetData()));
@@ -1315,9 +1314,11 @@ void FOpenXRHMD::EnumerateViews(FPipelinedFrameState& PipelineState)
 	{
 		TArray<XrViewConfigurationView> ViewConfigs;
 		Module->GetViewConfigurations(System, ViewConfigs);
+
+		const EStereoscopicPass PluginPassType = ViewConfigs.Num() > 1 ? EStereoscopicPass::eSSP_PRIMARY : EStereoscopicPass::eSSP_FULL;
 		for (int32 i = 0; i < ViewConfigs.Num(); i++)
 		{
-			PipelineState.PluginViews.Add(Module);
+			PipelineState.PluginViewInfos.Add({ Module, PluginPassType, true });
 		}
 		PipelineState.ViewConfigs.Append(ViewConfigs);
 	}
@@ -1331,7 +1332,12 @@ void FOpenXRHMD::EnumerateViews(FPipelinedFrameState& PipelineState)
 		FReadScopeLock DeviceLock(DeviceMutex);
 		for (IOpenXRExtensionPlugin* Module : ExtensionPlugins)
 		{
-			if (PipelineState.PluginViews.Contains(Module))
+			auto Predicate = [Module](const FPluginViewInfo& Info) -> bool
+			{
+				return Info.Plugin == Module;
+			};
+
+			if (PipelineState.PluginViewInfos.ContainsByPredicate(Predicate))
 			{
 				TArray<XrView> Views;
 				Module->GetViewLocations(Session, PipelineState.FrameState.predictedDisplayTime, DeviceSpaces[HMDDeviceId].Space, Views);
@@ -2548,9 +2554,9 @@ void FOpenXRHMD::OnFinishRendering_RHIThread()
 		{
 			TArray<XrSwapchainSubImage> ColorImages;
 			TArray<XrSwapchainSubImage> DepthImages;
-			for (int32 i = 0; i < PipelinedFrameStateRHI.PluginViews.Num(); i++)
+			for (int32 i = 0; i < PipelinedFrameStateRHI.PluginViewInfos.Num(); i++)
 			{
-				if (PipelinedFrameStateRHI.PluginViews[i] == Module && PipelinedLayerStateRHI.ColorImages.IsValidIndex(i))
+				if (PipelinedFrameStateRHI.PluginViewInfos[i].Plugin == Module && PipelinedLayerStateRHI.ColorImages.IsValidIndex(i))
 				{
 					ColorImages.Add(PipelinedLayerStateRHI.ColorImages[i]);
 					if (bDepthExtensionSupported)
diff --git a/Engine/Plugins/Runtime/OpenXR/Source/OpenXRHMD/Private/OpenXRHMD.h b/Engine/Plugins/Runtime/OpenXR/Source/OpenXRHMD/Private/OpenXRHMD.h
--- a/Engine/Plugins/Runtime/OpenXR/Source/OpenXRHMD/Private/OpenXRHMD.h
+++ b/Engine/Plugins/Runtime/OpenXR/Source/OpenXRHMD/Private/OpenXRHMD.h
@@ -65,6 +65,13 @@ public:
 		XrPosef BasePose;
 	};
 
+	struct FPluginViewInfo
+	{
+		class IOpenXRExtensionPlugin* Plugin = nullptr;
+		EStereoscopicPass PassType = EStereoscopicPass::eSSP_PRIMARY;
+		bool bIsPluginManaged = false;
+	};
+
 	// The game and render threads each have a separate copy of these structures so that they don't stomp on each other or cause tearing
 	// when the game thread progresses to the next frame while the render thread is still working on the previous frame.
 	struct FPipelinedFrameState
@@ -78,7 +85,7 @@ public:
 		float PixelDensity = 1.0f;
 
 		TArray<XrViewConfigurationView> ViewConfigs;
-		TArray<class IOpenXRExtensionPlugin*> PluginViews;
+		TArray<FPluginViewInfo> PluginViewInfos;
 
 		bool bXrFrameStateUpdated = false;
 	};
