Update comments in OpenXRHMD module

diff --git a/Engine/Plugins/Runtime/OpenXR/Source/OpenXRHMD/Private/OpenXRHMDModule.cpp b/Engine/Plugins/Runtime/OpenXR/Source/OpenXRHMD/Private/OpenXRHMDModule.cpp
--- a/Engine/Plugins/Runtime/OpenXR/Source/OpenXRHMD/Private/OpenXRHMDModule.cpp
+++ b/Engine/Plugins/Runtime/OpenXR/Source/OpenXRHMD/Private/OpenXRHMDModule.cpp
@@ -605,10 +605,10 @@ bool FOpenXRHMDModule::InitInstance()
 
 	// Enable layers, if specified by CVar.
 	// Note: For the validation layer to work on Windows (as of latest OpenXR runtime, August 2019), the following are required:
-	//   1. Download and build the OpenXR SDK from [REDACTED_PROFILE_URL]/OpenXR-SDK-Source (follow instructions at [REDACTED_PROFILE_URL]/OpenXR-SDK-Source/blob/master/BUILDING.md)
+	//   1. Download and build the OpenXR SDK from [REDACTED_PROFILE_URL]/OpenXR-SDK-Source (follow instructions at [REDACTED_PROFILE_URL]/OpenXR-SDK-Source/blob/main/BUILDING.md)
 	//	 2. Add a registry key under HKEY_LOCAL_MACHINE\SOFTWARE\Khronos\OpenXR\1\ApiLayers\Explicit, containing the path to the manifest file
-	//      (e.g. C:\OpenXR-SDK-Source-master\build\win64\src\api_layers\XrApiLayer_core_validation.json) <-- this file is downloaded as part of the SDK source, above
-	//   3. Copy the DLL from the build target at, for example, C:\OpenXR-SDK-Source-master\build\win64\src\api_layers\XrApiLayer_core_validation.dll to
+	//      (e.g. C:\OpenXR-SDK-Source-main\build\win64\src\api_layers\XrApiLayer_core_validation.json) <-- this file is downloaded as part of the SDK source, above
+	//   3. Copy the DLL from the build target at, for example, C:\OpenXR-SDK-Source-main\build\win64\src\api_layers\XrApiLayer_core_validation.dll to
 	//      somewhere in your system path (e.g. c:\windows\system32); the OpenXR loader currently doesn't use the path the json file is in (this is a bug)
 
 	const bool bEnableOpenXRValidationLayer = CVarEnableOpenXRValidationLayer.GetValueOnAnyThread() != 0;
