UE-129957 OpenXR: Playspace
-Adding features related to play area bounds to OpenXR.
-New Blueprint to get the bounds and their origin in UE world space.
-New Delegate to respond to play area updates.

diff --git a/Engine/Plugins/Runtime/OpenXR/Source/OpenXRHMD/Private/OpenXRHMD.cpp b/Engine/Plugins/Runtime/OpenXR/Source/OpenXRHMD/Private/OpenXRHMD.cpp
--- a/Engine/Plugins/Runtime/OpenXR/Source/OpenXRHMD/Private/OpenXRHMD.cpp
+++ b/Engine/Plugins/Runtime/OpenXR/Source/OpenXRHMD/Private/OpenXRHMD.cpp
@@ -366,6 +366,104 @@ FVector2D FOpenXRHMD::GetPlayAreaBounds(EHMDTrackingOrigin::Type Origin) const
 	return ToFVector2D(Bounds, WorldToMetersScale);
 }
 
+bool FOpenXRHMD::GetPlayAreaRect(FTransform& OutTransform, FVector2D& OutRect) const
+{
+	// Get the origin and the extents of the play area rect.
+	// The OpenXR Stage Space defines the origin of the playable rectangle.  The origin is at the floor. xrGetReferenceSpaceBoundsRect will give you the horizontal extents.
+
+	const FPipelinedFrameState& PipelinedState = GetPipelinedFrameStateForThread();
+
+	{
+		if (StageSpace == XR_NULL_HANDLE)
+		{
+			return false;
+		}
+
+		XrSpaceLocation NewLocation = { XR_TYPE_SPACE_LOCATION };
+		const XrResult Result = xrLocateSpace(StageSpace, PipelinedState.TrackingSpace->Handle, PipelinedState.FrameState.predictedDisplayTime, &NewLocation);
+		if (Result != XR_SUCCESS)
+		{
+			return false;
+		}
+
+		if (!(NewLocation.locationFlags & (XR_SPACE_LOCATION_POSITION_VALID_BIT | XR_SPACE_LOCATION_ORIENTATION_VALID_BIT)))
+		{
+			return false;
+		}
+		const FQuat Orientation = ToFQuat(NewLocation.pose.orientation);
+		const FVector Position = ToFVector(NewLocation.pose.position, PipelinedState.WorldToMetersScale);
+
+		FTransform TrackingToWorld = GetTrackingToWorldTransform();
+		OutTransform = FTransform(Orientation, Position) * TrackingToWorld;
+	}
+
+	{
+		XrExtent2Df Bounds; // width is X height is Z
+		const XrResult Result = xrGetReferenceSpaceBoundsRect(Session, XR_REFERENCE_SPACE_TYPE_STAGE, &Bounds);
+		if (Result != XR_SUCCESS)
+		{
+			return false;
+		}
+
+		OutRect = ToFVector2D(Bounds, PipelinedState.WorldToMetersScale);
+	}
+
+	return true;
+}
+
+bool FOpenXRHMD::GetTrackingOriginTransform(TEnumAsByte<EHMDTrackingOrigin::Type> Origin, FTransform& OutTransform)  const
+{
+	XrSpace Space = XR_NULL_HANDLE;
+	switch (Origin)
+	{
+	case EHMDTrackingOrigin::Eye:
+		if (DeviceSpaces.Num())
+		{
+			Space = DeviceSpaces[HMDDeviceId].Space;
+		}
+		break;
+	case EHMDTrackingOrigin::Floor:
+		Space = LocalSpace;
+		break;
+	case EHMDTrackingOrigin::Stage:
+		Space = StageSpace;
+		break;
+	//case EHMDTrackingOrigin::???:
+		//Space = CustomSpace
+		//break;
+	default:
+		check(false);
+		break;
+	}
+
+	if (Space == XR_NULL_HANDLE)
+	{
+		// This space is not supported.
+		return false;
+	}
+
+	const FPipelinedFrameState& PipelinedState = GetPipelinedFrameStateForThread();
+
+	XrSpaceLocation NewLocation = { XR_TYPE_SPACE_LOCATION };
+	const XrResult Result = xrLocateSpace(Space, PipelinedState.TrackingSpace->Handle, PipelinedState.FrameState.predictedDisplayTime, &NewLocation);
+	if (Result != XR_SUCCESS)
+	{
+		return false;
+	}
+	if (!(NewLocation.locationFlags & (XR_SPACE_LOCATION_POSITION_VALID_BIT | XR_SPACE_LOCATION_ORIENTATION_VALID_BIT)))
+	{
+		return false;
+	}
+	const FQuat Orientation = ToFQuat(NewLocation.pose.orientation);
+	const FVector Position = ToFVector(NewLocation.pose.position, PipelinedState.WorldToMetersScale);
+
+	FTransform TrackingToWorld = GetTrackingToWorldTransform();
+	OutTransform =  FTransform(Orientation, Position) * TrackingToWorld;
+
+	return true;
+}
+
+
 FName FOpenXRHMD::GetHMDName() const
 {
 	return SystemProperties.systemName;
@@ -2377,6 +2475,12 @@ bool FOpenXRHMD::OnStartGameFrame(FWorldContext& WorldContext)
 				BasePosition = FVector::ZeroVector;
 				bTrackingSpaceInvalid = true;
 			}
+
+			if (SpaceChange.referenceSpaceType == XR_REFERENCE_SPACE_TYPE_STAGE)
+			{
+				OnPlayAreaChanged();
+			}
+
 			break;
 		}
 		case XR_TYPE_EVENT_DATA_VISIBILITY_MASK_CHANGED_KHR:
diff --git a/Engine/Plugins/Runtime/OpenXR/Source/OpenXRHMD/Private/OpenXRHMD.h b/Engine/Plugins/Runtime/OpenXR/Source/OpenXRHMD/Private/OpenXRHMD.h
--- a/Engine/Plugins/Runtime/OpenXR/Source/OpenXRHMD/Private/OpenXRHMD.h
+++ b/Engine/Plugins/Runtime/OpenXR/Source/OpenXRHMD/Private/OpenXRHMD.h
@@ -207,6 +207,8 @@ public:
 	virtual float GetWorldToMetersScale() const override;
 
 	virtual FVector2D GetPlayAreaBounds(EHMDTrackingOrigin::Type Origin) const override;
+	virtual bool GetPlayAreaRect(FTransform& OutTransform, FVector2D& OutExtent) const override;
+	virtual bool GetTrackingOriginTransform(TEnumAsByte<EHMDTrackingOrigin::Type> Origin, FTransform& OutTransform)  const override;
 
 	virtual bool HDRGetMetaDataForStereo(EDisplayOutputFormat& OutDisplayOutputFormat, EDisplayColorGamut& OutDisplayColorGamut, bool& OutbHDRSupported) override;
 
diff --git a/Engine/Source/Runtime/HeadMountedDisplay/Private/HeadMountedDisplayFunctionLibrary.cpp b/Engine/Source/Runtime/HeadMountedDisplay/Private/HeadMountedDisplayFunctionLibrary.cpp
--- a/Engine/Source/Runtime/HeadMountedDisplay/Private/HeadMountedDisplayFunctionLibrary.cpp
+++ b/Engine/Source/Runtime/HeadMountedDisplay/Private/HeadMountedDisplayFunctionLibrary.cpp
@@ -595,6 +595,25 @@ FVector2D UHeadMountedDisplayFunctionLibrary::GetPlayAreaBounds(TEnumAsByte<EHMD
 	return FVector2D::ZeroVector;
 }
 
+bool UHeadMountedDisplayFunctionLibrary::GetTrackingOriginTransform(TEnumAsByte<EHMDTrackingOrigin::Type> Origin, FTransform& OutTransform)
+{
+	if (GEngine->XRSystem.IsValid())
+	{
+		return GEngine->XRSystem->GetTrackingOriginTransform(Origin, OutTransform);
+	}
+	return false;
+}
+
+bool UHeadMountedDisplayFunctionLibrary::GetPlayAreaRect(FTransform& OutTransform, FVector2D& OutRect)
+{
+	if (GEngine->XRSystem.IsValid())
+	{
+		return GEngine->XRSystem->GetPlayAreaRect(OutTransform, OutRect);
+	}
+	return false;
+}
+
+
 void UHeadMountedDisplayFunctionLibrary::BreakKey(FKey InKey, FString& InteractionProfile, EControllerHand& Hand, FName& MotionSource, FString& Indentifier, FString& Component)
 {
 	TArray<FString> Tokens;
diff --git a/Engine/Source/Runtime/HeadMountedDisplay/Private/VRNotificationsComponent.cpp b/Engine/Source/Runtime/HeadMountedDisplay/Private/VRNotificationsComponent.cpp
--- a/Engine/Source/Runtime/HeadMountedDisplay/Private/VRNotificationsComponent.cpp
+++ b/Engine/Source/Runtime/HeadMountedDisplay/Private/VRNotificationsComponent.cpp
@@ -3,6 +3,7 @@
 
 #include "VRNotificationsComponent.h"
 #include "Misc/CoreDelegates.h"
+#include "XRTrackingSystemBase.h"
 
 UVRNotificationsComponent::UVRNotificationsComponent(const FObjectInitializer& ObjectInitializer)
 	: Super(ObjectInitializer)
@@ -22,6 +23,8 @@ void UVRNotificationsComponent::OnRegister()
 	FCoreDelegates::VRHeadsetPutOnHead.AddUObject(this, &UVRNotificationsComponent::HMDPutOnHeadDelegate_Handler);
 	FCoreDelegates::VRHeadsetRemovedFromHead.AddUObject(this, &UVRNotificationsComponent::HMDRemovedFromHeadDelegate_Handler);
 	FCoreDelegates::VRControllerRecentered.AddUObject(this, &UVRNotificationsComponent::VRControllerRecentered_Handler);
+	FXRTrackingSystemDelegates::OnXRTrackingOriginChanged.AddUObject(this, &UVRNotificationsComponent::XRTrackingOriginChanged_Handler);
+	FXRTrackingSystemDelegates::OnXRPlayAreaChanged.AddUObject(this, &UVRNotificationsComponent::XRPlayAreaChanged_Handler);
 }
 
 void UVRNotificationsComponent::OnUnregister()
@@ -37,4 +40,6 @@ void UVRNotificationsComponent::OnUnregister()
 	FCoreDelegates::VRHeadsetPutOnHead.RemoveAll(this);
 	FCoreDelegates::VRHeadsetRemovedFromHead.RemoveAll(this);
 	FCoreDelegates::VRControllerRecentered.RemoveAll(this);
+	FXRTrackingSystemDelegates::OnXRTrackingOriginChanged.RemoveAll(this);
+	FXRTrackingSystemDelegates::OnXRPlayAreaChanged.RemoveAll(this);
 }
diff --git a/Engine/Source/Runtime/HeadMountedDisplay/Private/XRTrackingSystemBase.cpp b/Engine/Source/Runtime/HeadMountedDisplay/Private/XRTrackingSystemBase.cpp
--- a/Engine/Source/Runtime/HeadMountedDisplay/Private/XRTrackingSystemBase.cpp
+++ b/Engine/Source/Runtime/HeadMountedDisplay/Private/XRTrackingSystemBase.cpp
@@ -13,6 +13,7 @@
 
 // Tracking system delegates
 FXRTrackingSystemDelegates::FXRTrackingOriginChanged FXRTrackingSystemDelegates::OnXRTrackingOriginChanged;
+FXRTrackingSystemDelegates::FXRPlayAreaChanged FXRTrackingSystemDelegates::OnXRPlayAreaChanged;
 
 
 
diff --git a/Engine/Source/Runtime/HeadMountedDisplay/Public/HeadMountedDisplayFunctionLibrary.h b/Engine/Source/Runtime/HeadMountedDisplay/Public/HeadMountedDisplayFunctionLibrary.h
--- a/Engine/Source/Runtime/HeadMountedDisplay/Public/HeadMountedDisplayFunctionLibrary.h
+++ b/Engine/Source/Runtime/HeadMountedDisplay/Public/HeadMountedDisplayFunctionLibrary.h
@@ -380,6 +380,18 @@ class HEADMOUNTEDDISPLAY_API UHeadMountedDisplayFunctionLibrary : public UBluepr
 	UFUNCTION(BlueprintCallable, Category = "Input|XRTracking")
 	static FVector2D GetPlayAreaBounds(TEnumAsByte<EHMDTrackingOrigin::Type> Origin = EHMDTrackingOrigin::Stage);
 
+	/**
+	 * Get the transform of the specified tracking origin, if available.
+	 */
+	UFUNCTION(BlueprintCallable, Category = "Input|XRTracking")
+	static bool GetTrackingOriginTransform(TEnumAsByte<EHMDTrackingOrigin::Type> Origin, FTransform& OutTransform);
+
+	/**
+	 * Get the transform and dimensions of the playable area rectangle.  Returns false if none currently specified/available.
+	 */
+	UFUNCTION(BlueprintCallable, Category = "Input|XRTracking")
+	static bool GetPlayAreaRect(FTransform& OutTransform, FVector2D& OutRect);
+
 	/** Breaks an XR key apart into the interaction profile, handedness, motion source, indentifier and component. */
 	UFUNCTION(BlueprintPure, Category = "Input|XRTracking", meta = (NativeBreakFunc))
 	static void BreakKey(FKey InKey, FString& InteractionProfile, EControllerHand& Hand, FName& MotionSource, FString& Indentifier, FString& Component);
diff --git a/Engine/Source/Runtime/HeadMountedDisplay/Public/IXRTrackingSystem.h b/Engine/Source/Runtime/HeadMountedDisplay/Public/IXRTrackingSystem.h
--- a/Engine/Source/Runtime/HeadMountedDisplay/Public/IXRTrackingSystem.h
+++ b/Engine/Source/Runtime/HeadMountedDisplay/Public/IXRTrackingSystem.h
@@ -436,4 +436,14 @@ public:
 	 * Get the bounds of the area where the user can freely move while remaining tracked centered around the specified origin
 	 */
 	virtual FVector2D GetPlayAreaBounds(EHMDTrackingOrigin::Type Origin) const { return FVector2D::ZeroVector; }
+
+	/**
+	 * Get the transform of the specified tracking origin, if available.
+	 */
+	virtual bool GetTrackingOriginTransform(TEnumAsByte<EHMDTrackingOrigin::Type> Origin, FTransform& OutTransform) const { return false; }
+
+	/**
+	 * Get the transform and dimensions of the area where the user can freely move while remaining tracked centered around the specified origin transform
+	 */
+	virtual bool GetPlayAreaRect(FTransform& OutTransform, FVector2D& OutRect) const { return false; }
 };
diff --git a/Engine/Source/Runtime/HeadMountedDisplay/Public/VRNotificationsComponent.h b/Engine/Source/Runtime/HeadMountedDisplay/Public/VRNotificationsComponent.h
--- a/Engine/Source/Runtime/HeadMountedDisplay/Public/VRNotificationsComponent.h
+++ b/Engine/Source/Runtime/HeadMountedDisplay/Public/VRNotificationsComponent.h
@@ -8,6 +8,8 @@
 #include "Components/ActorComponent.h"
 #include "VRNotificationsComponent.generated.h"
 
+class IXRTrackingSystem;
+
 UCLASS(Blueprintable, meta = (BlueprintSpawnableComponent), ClassGroup = HeadMountedDisplay)
 class HEADMOUNTEDDISPLAY_API UVRNotificationsComponent : public UActorComponent
 {
@@ -57,6 +59,14 @@ class HEADMOUNTEDDISPLAY_API UVRNotificationsComponent : public UActorComponent
 	UPROPERTY(BlueprintAssignable)
 	FVRNotificationsDelegate VRControllerRecenteredDelegate;
 
+	// This will be called when the VR system recenters a controller.  
+	UPROPERTY(BlueprintAssignable)
+	FVRNotificationsDelegate XRTrackingOriginChangedDelegate;
+
+	// This will be called when the VR system recenters a controller.  
+	UPROPERTY(BlueprintAssignable)
+	FVRNotificationsDelegate XRPlayAreaChangedDelegate;
+
 public:
 	void OnRegister() override;
 	void OnUnregister() override;
@@ -65,13 +75,15 @@ private:
 	/** Native handlers that get registered with the actual FCoreDelegates, and then proceed to broadcast to the delegates above */
 	void HMDTrackingInitializingAndNeedsHMDToBeTrackedDelegate_Handler()	{ HMDTrackingInitializingAndNeedsHMDToBeTrackedDelegate.Broadcast(); }
 	void HMDTrackingInitializedDelegate_Handler()	{ HMDTrackingInitializedDelegate.Broadcast(); }
-	void HMDRecenteredDelegate_Handler()	{ HMDRecenteredDelegate.Broadcast(); }
-	void HMDLostDelegate_Handler()			{ HMDLostDelegate.Broadcast(); }
-	void HMDReconnectedDelegate_Handler()	{ HMDReconnectedDelegate.Broadcast(); }
-	void HMDConnectCanceledDelegate_Handler() { HMDConnectCanceledDelegate.Broadcast(); }
-	void HMDPutOnHeadDelegate_Handler() { HMDPutOnHeadDelegate.Broadcast(); }
-	void HMDRemovedFromHeadDelegate_Handler() { HMDRemovedFromHeadDelegate.Broadcast(); }
-	void VRControllerRecentered_Handler() { VRControllerRecenteredDelegate.Broadcast(); }
+	void HMDRecenteredDelegate_Handler()		{ HMDRecenteredDelegate.Broadcast(); }
+	void HMDLostDelegate_Handler()				{ HMDLostDelegate.Broadcast(); }
+	void HMDReconnectedDelegate_Handler()		{ HMDReconnectedDelegate.Broadcast(); }
+	void HMDConnectCanceledDelegate_Handler()	{ HMDConnectCanceledDelegate.Broadcast(); }
+	void HMDPutOnHeadDelegate_Handler()			{ HMDPutOnHeadDelegate.Broadcast(); }
+	void HMDRemovedFromHeadDelegate_Handler()	{ HMDRemovedFromHeadDelegate.Broadcast(); }
+	void VRControllerRecentered_Handler()		{ VRControllerRecenteredDelegate.Broadcast(); }
+	void XRTrackingOriginChanged_Handler(const IXRTrackingSystem* TrackingSys)		{ XRTrackingOriginChangedDelegate.Broadcast(); }
+	void XRPlayAreaChanged_Handler()			{ XRPlayAreaChangedDelegate.Broadcast(); }
 };
 
 
diff --git a/Engine/Source/Runtime/HeadMountedDisplay/Public/XRTrackingSystemBase.h b/Engine/Source/Runtime/HeadMountedDisplay/Public/XRTrackingSystemBase.h
--- a/Engine/Source/Runtime/HeadMountedDisplay/Public/XRTrackingSystemBase.h
+++ b/Engine/Source/Runtime/HeadMountedDisplay/Public/XRTrackingSystemBase.h
@@ -11,6 +11,9 @@ class HEADMOUNTEDDISPLAY_API FXRTrackingSystemDelegates
 public:
 	DECLARE_MULTICAST_DELEGATE_OneParam(FXRTrackingOriginChanged, const IXRTrackingSystem* /*TrackingSystem*/);
 	static FXRTrackingOriginChanged OnXRTrackingOriginChanged;
+
+	DECLARE_MULTICAST_DELEGATE(FXRPlayAreaChanged);
+	static FXRPlayAreaChanged OnXRPlayAreaChanged;
 };
 
 /** 
@@ -182,6 +185,11 @@ protected:
 		FXRTrackingSystemDelegates::OnXRTrackingOriginChanged.Broadcast(this);
 	}
 
+	virtual void OnPlayAreaChanged()
+	{
+		FXRTrackingSystemDelegates::OnXRPlayAreaChanged.Broadcast();
+	}
+
 	/**
 	 * Computes the project's tracking-to-world transform based off how the user 
 	 * has set up their camera system (assumes the camera is parented to the XR 
