OpenXR PixelDensity: vr.pixeldensity active

The OpenXR plugin now adjusts the stereo render targets based on vr.pixeldensity

* AdjustViewRect respects PixelDensity, plus quantizes in order to match RenderTarget size in renderer logic (prevents upscale in MobileShadingRenderer)

* Override CalculateRenderTargetSize - previously, FOpenXRHMD was using the default implementation from FXRRenderTargetManager.
  However, it doesn't care about the per-view quanitization utilized by AdjustViewRect. Instead of modifying the shared implementation,
  we can provide a customized OpenXR version that shares logic with AdjustViewRect

* Allocate the depth texture at color alloc time - Because of the location of where IStereoRenderTargetManager::AllocateDepthTexture is called (scene texture creation time),
  it's too late for us to acquire the swapchain because it is called _after_ OnBeginRendering_RenderThread. This means we would use janky logic to acquire the swapchain after
  the frame has started. Instead, since we know if we're going to use depth at session time, and the parameters should be 'known', we can just allocate depth alongside color.
  Then the interface `AllocateDepthTexture` call is a validator + getter. In the future, we could rename some of these interfaces.

* Update layer/swapchain state in PostRenderViewFamily_RenderThread - we only need this state updated for frame end, so we can update it just-in-time. This actually happens
  'early' becase we enqueue directly onto GraphBuilder.RHICmdList. In the future, we could schedule a task onto GraphBuilder.

diff --git a/Engine/Plugins/Runtime/OpenXR/Source/OpenXRHMD/Private/OpenXRHMD.cpp b/Engine/Plugins/Runtime/OpenXR/Source/OpenXRHMD/Private/OpenXRHMD.cpp
--- a/Engine/Plugins/Runtime/OpenXR/Source/OpenXRHMD/Private/OpenXRHMD.cpp
+++ b/Engine/Plugins/Runtime/OpenXR/Source/OpenXRHMD/Private/OpenXRHMD.cpp
@@ -836,6 +836,24 @@ bool FOpenXRHMD::EnableStereo(bool stereo)
 	}
 }
 
+FIntPoint GeneratePixelDensitySize(const XrViewConfigurationView& Config, const float PixelDensity)
+{
+	FIntPoint DensityAdjustedSize =
+	{
+		FMath::CeilToInt(Config.recommendedImageRectWidth * PixelDensity),
+		FMath::CeilToInt(Config.recommendedImageRectHeight * PixelDensity)
+	};
+
+	// We quantize in order to be consistent with the rest of the engine in creating our buffers.
+	// Interestingly, we need to be a bit careful with this quantization during target alloc because
+	// some runtime compositors want/expect targets that match the recommended size. Some runtimes
+	// might blit from a 'larger' size to the recommended size. This could happen with quantization
+	// factors that don't align with the recommended size.
+	QuantizeSceneBufferSize(DensityAdjustedSize, DensityAdjustedSize);
+
+	return DensityAdjustedSize;
+}
+
 void FOpenXRHMD::AdjustViewRect(int32 ViewIndex, int32& X, int32& Y, uint32& SizeX, uint32& SizeY) const
 {
 	const FPipelinedFrameState& PipelineState = GetPipelinedFrameStateForThread();
@@ -846,14 +864,47 @@ void FOpenXRHMD::AdjustViewRect(int32 ViewIndex, int32& X, int32& Y, uint32& Siz
 	// Thus the start index should be the second view if enabled
 	for (int32 i = bIsMobileMultiViewEnabled ? 1 : 0; i < ViewIndex; ++i)
 	{
-		ViewRectMin.X += PipelineState.ViewConfigs[i].recommendedImageRectWidth;
+		ViewRectMin.X += FMath::CeilToInt(PipelineState.ViewConfigs[i].recommendedImageRectWidth * PipelineState.PixelDensity);
+		QuantizeSceneBufferSize(ViewRectMin, ViewRectMin);
 	}
-	QuantizeSceneBufferSize(ViewRectMin, ViewRectMin);
 
 	X = ViewRectMin.X;
 	Y = ViewRectMin.Y;
-	SizeX = Config.recommendedImageRectWidth;
-	SizeY = Config.recommendedImageRectHeight;
+
+	const FIntPoint DensityAdjustedSize = GeneratePixelDensitySize(Config, PipelineState.PixelDensity);
+
+	SizeX = DensityAdjustedSize.X;
+	SizeY = DensityAdjustedSize.Y;
+
+}
+
+void FOpenXRHMD::CalculateRenderTargetSize(const FViewport& Viewport, uint32& InOutSizeX, uint32& InOutSizeY)
+{
+	check(IsInGameThread() || IsInRenderingThread());
+
+	const FPipelinedFrameState& PipelineState = GetPipelinedFrameStateForThread();
+	const float PixelDensity = PipelineState.PixelDensity;
+
+	// TODO: Could we just call AdjustViewRect per view, or even for _only_ the last view?
+	FIntPoint Size(EForceInit::ForceInitToZero);
+	for (int32 ViewIndex = 0; ViewIndex < PipelineState.ViewConfigs.Num(); ViewIndex++)
+	{
+		const XrViewConfigurationView& Config = PipelineState.ViewConfigs[ViewIndex];
+
+		// If Mobile Multi-View is active the first two views will share the same position
+		// TODO: This is weird logic that we should re-investigate. It makes sense for AdjustViewRect, but not
+		// for the 'size' of an RT.
+		const bool bMMVView = bIsMobileMultiViewEnabled && ViewIndex < 2;
+
+		const FIntPoint DensityAdjustedSize = GeneratePixelDensitySize(Config, PipelineState.PixelDensity);
+		Size.X = bMMVView ? FMath::Max(Size.X, DensityAdjustedSize.X) : Size.X + DensityAdjustedSize.X;
+		Size.Y = FMath::Max(Size.Y, DensityAdjustedSize.Y);
+	}
+
+	InOutSizeX = Size.X;
+	InOutSizeY = Size.Y;
+
+	check(InOutSizeX != 0 && InOutSizeY != 0);
 }
 
 void FOpenXRHMD::SetFinalViewRect(FRHICommandListImmediate& RHICmdList, const int32 ViewIndex, const FIntRect& FinalViewRect)
@@ -865,10 +916,7 @@ void FOpenXRHMD::SetFinalViewRect(FRHICommandListImmediate& RHICmdList, const in
 		return;
 	}
 
-	float NearZ = GNearClippingPlane / GetWorldToMetersScale();
-
 	XrSwapchainSubImage& ColorImage = PipelinedLayerStateRendering.ColorImages[ViewIndex];
-	ColorImage.swapchain = PipelinedLayerStateRendering.ColorSwapchain.IsValid() ? static_cast<FOpenXRSwapchain*>(PipelinedLayerStateRendering.ColorSwapchain.Get())->GetHandle() : XR_NULL_HANDLE;
 	ColorImage.imageArrayIndex = bIsMobileMultiViewEnabled && ViewIndex < 2 ? ViewIndex : 0;
 	ColorImage.imageRect = {
 		{ FinalViewRect.Min.X, FinalViewRect.Min.Y },
@@ -878,56 +926,9 @@ void FOpenXRHMD::SetFinalViewRect(FRHICommandListImmediate& RHICmdList, const in
 	XrSwapchainSubImage& DepthImage = PipelinedLayerStateRendering.DepthImages[ViewIndex];
 	if (bDepthExtensionSupported)
 	{
-		DepthImage.swapchain = PipelinedLayerStateRendering.DepthSwapchain.IsValid() ? static_cast<FOpenXRSwapchain*>(PipelinedLayerStateRendering.DepthSwapchain.Get())->GetHandle() : XR_NULL_HANDLE;
 		DepthImage.imageArrayIndex = bIsMobileMultiViewEnabled && ViewIndex < 2 ? ViewIndex : 0;
 		DepthImage.imageRect = ColorImage.imageRect;
 	}
-
-	if (!PipelinedFrameStateRendering.PluginViewInfos.IsValidIndex(ViewIndex))
-	{
-		// This plugin is no longer providing this view.
-		return;
-	}
-
-	if (PipelinedFrameStateRendering.PluginViewInfos[ViewIndex].bIsPluginManaged)
-	{
-		// Defer to the plugin to handle submission
-		return;
-	}
-
-	XrCompositionLayerProjectionView& Projection = PipelinedLayerStateRendering.ProjectionLayers[ViewIndex];
-	XrCompositionLayerDepthInfoKHR& DepthLayer = PipelinedLayerStateRendering.DepthLayers[ViewIndex];
-
-	Projection.type = XR_TYPE_COMPOSITION_LAYER_PROJECTION_VIEW;
-	Projection.next = nullptr;
-	Projection.subImage = ColorImage;
-
-	if (bDepthExtensionSupported && PipelinedLayerStateRendering.DepthSwapchain.IsValid())
-	{
-		DepthLayer.type = XR_TYPE_COMPOSITION_LAYER_DEPTH_INFO_KHR;
-		DepthLayer.next = nullptr;
-		DepthLayer.subImage = DepthImage;
-		DepthLayer.minDepth = 0.0f;
-		DepthLayer.maxDepth = 1.0f;
-		DepthLayer.nearZ = FLT_MAX;
-		DepthLayer.farZ = NearZ;
-
-		for (IOpenXRExtensionPlugin* Module : ExtensionPlugins)
-		{
-			DepthLayer.next = Module->OnBeginDepthInfo(Session, 0, ViewIndex, DepthLayer.next);
-		}
-		Projection.next = &DepthLayer;
-	}
-
-	for (IOpenXRExtensionPlugin* Module : ExtensionPlugins)
-	{
-		Projection.next = Module->OnBeginProjectionView(Session, 0, ViewIndex, Projection.next);
-	}
-
-	RHICmdList.EnqueueLambda([this, LayerState = PipelinedLayerStateRendering](FRHICommandListImmediate& InRHICmdList)
-	{
-		PipelinedLayerStateRHI = LayerState;
-	});
 }
 
 EStereoscopicPass FOpenXRHMD::GetViewPassForIndex(bool bStereoRequested, int32 ViewIndex) const
@@ -1091,6 +1092,74 @@ void FOpenXRHMD::PreRenderViewFamily_RenderThread(FRDGBuilder& GraphBuilder, FSc
 	}
 }
 
+void FOpenXRHMD::PostRenderViewFamily_RenderThread(FRDGBuilder& GraphBuilder, FSceneViewFamily& InViewFamily)
+{
+	check(IsInRenderingThread());
+
+	const float NearZ = GNearClippingPlane / GetWorldToMetersScale();
+
+	for (int32 ViewIndex = 0; ViewIndex < PipelinedLayerStateRendering.ColorImages.Num(); ViewIndex++)
+	{
+		if (!PipelinedLayerStateRendering.ColorImages.IsValidIndex(ViewIndex))
+		{
+			continue;
+		}
+
+		// Update SubImages with latest swapchain
+		XrSwapchainSubImage& ColorImage = PipelinedLayerStateRendering.ColorImages[ViewIndex];
+		XrSwapchainSubImage& DepthImage = PipelinedLayerStateRendering.DepthImages[ViewIndex];
+
+		ColorImage.swapchain = PipelinedLayerStateRendering.ColorSwapchain.IsValid() ? static_cast<FOpenXRSwapchain*>(PipelinedLayerStateRendering.ColorSwapchain.Get())->GetHandle() : XR_NULL_HANDLE;
+		if (bDepthExtensionSupported)
+		{
+			DepthImage.swapchain = PipelinedLayerStateRendering.DepthSwapchain.IsValid() ? static_cast<FOpenXRSwapchain*>(PipelinedLayerStateRendering.DepthSwapchain.Get())->GetHandle() : XR_NULL_HANDLE;
+		}
+
+		if (IsViewManagedByPlugin(ViewIndex))
+		{
+			// Plugin owns further usage of the subimage, so we don't use the subimages in our layers
+			continue;
+		}
+
+		XrCompositionLayerProjectionView& Projection = PipelinedLayerStateRendering.ProjectionLayers[ViewIndex];
+
+		Projection.type = XR_TYPE_COMPOSITION_LAYER_PROJECTION_VIEW;
+		Projection.next = nullptr;
+		Projection.subImage = ColorImage;
+
+		if (bDepthExtensionSupported && PipelinedLayerStateRendering.DepthSwapchain.IsValid())
+		{
+			XrCompositionLayerDepthInfoKHR& DepthLayer = PipelinedLayerStateRendering.DepthLayers[ViewIndex];
+
+			DepthLayer.type = XR_TYPE_COMPOSITION_LAYER_DEPTH_INFO_KHR;
+			DepthLayer.next = nullptr;
+			DepthLayer.subImage = DepthImage;
+			DepthLayer.minDepth = 0.0f;
+			DepthLayer.maxDepth = 1.0f;
+			DepthLayer.nearZ = FLT_MAX;
+			DepthLayer.farZ = NearZ;
+
+			for (IOpenXRExtensionPlugin* Module : ExtensionPlugins)
+			{
+				DepthLayer.next = Module->OnBeginDepthInfo(Session, 0, ViewIndex, DepthLayer.next);
+			}
+
+			Projection.next = &DepthLayer;
+		}
+
+		for (IOpenXRExtensionPlugin* Module : ExtensionPlugins)
+		{
+			Projection.next = Module->OnBeginProjectionView(Session, 0, ViewIndex, Projection.next);
+		}
+	}
+
+	// We use RHICmdList directly, though eventually, we might want to schedule on GraphBuilder
+	GraphBuilder.RHICmdList.EnqueueLambda([this, LayerState = PipelinedLayerStateRendering](FRHICommandListImmediate&)
+	{
+		PipelinedLayerStateRHI = LayerState;
+	});
+}
+
 bool FOpenXRHMD::IsActiveThisFrame_Internal(const FSceneViewExtensionContext& Context) const
 {
 	// Don't activate the SVE if xr is being used for tracking only purposes
@@ -1144,7 +1213,7 @@ FOpenXRHMD::FOpenXRHMD(const FAutoRegister& AutoRegister, XrInstance InInstance,
 	, RenderBridge(InRenderBridge)
 	, RendererModule(nullptr)
 	, LastRequestedSwapchainFormat(0)
-	, LastRequestedDepthSwapchainFormat(0)
+	, LastRequestedDepthSwapchainFormat(PF_DepthStencil)
 	, bTrackingSpaceInvalid(true)
 	, bUseCustomReferenceSpace(false)
 	, BaseOrientation(FQuat::Identity)
@@ -1472,6 +1541,18 @@ void FOpenXRHMD::EnumerateViews(FPipelinedFrameState& PipelineState)
 	}
 }
 
+bool FOpenXRHMD::IsViewManagedByPlugin(int32 ViewIndex) const
+{
+	check(IsInRenderingThread());
+	if (!PipelinedFrameStateRendering.PluginViewInfos.IsValidIndex(ViewIndex))
+	{
+		// Was formerly associated with plugin, but plugin is no longer providing this view
+		// Core views always have an entry in the PluginViewInfos array
+		return true;
+	}
+	return PipelinedFrameStateRendering.PluginViewInfos[ViewIndex].bIsPluginManaged;
+}
+
 #if !PLATFORM_HOLOLENS
 void FOpenXRHMD::BuildOcclusionMeshes()
 {
@@ -1997,55 +2078,94 @@ bool FOpenXRHMD::AllocateRenderTargetTexture(uint32 Index, uint32 SizeX, uint32
 	OutTargetableTexture = OutShaderResourceTexture = (FTexture2DRHIRef&)Swapchain->GetTextureRef();
 	LastRequestedSwapchainFormat = Format;
 
-	bNeedReAllocatedDepth = bDepthExtensionSupported;
+	// TODO: Pass in known depth parameters (format + flags)? Do we know that at viewport setup time?
+	AllocateDepthTextureInternal(Index, SizeX, SizeY, NumSamples);
 
 	return true;
 }
 
-bool FOpenXRHMD::AllocateDepthTexture(uint32 Index, uint32 SizeX, uint32 SizeY, uint8 Format, uint32 NumMips, ETextureCreateFlags Flags, ETextureCreateFlags TargetableTextureFlags, FTexture2DRHIRef& OutTargetableTexture, FTexture2DRHIRef& OutShaderResourceTexture, uint32 NumSamples)
+void FOpenXRHMD::AllocateDepthTextureInternal(uint32 Index, uint32 SizeX, uint32 SizeY, uint32 NumSamples)
 {
+	// TODO: Allocate depth texture by checking bNeedReAllocatedDepth
+
 	check(IsInRenderingThread());
 
-	// FIXME: UE constantly calls this function even when there is no reason to reallocate the depth texture (see NeedReAllocateDepthTexture)
 	FReadScopeLock Lock(SessionHandleMutex);
 	if (!Session || !bDepthExtensionSupported)
 	{
-		return false;
+		return;
 	}
 
-	// We're only creating a 1x target here, but we don't know whether it'll be the targeted texture
-	// or the resolve texture. Because of this, we unify the input flags.
-	ETextureCreateFlags UnifiedCreateFlags = Flags | TargetableTextureFlags;
+	check(LastRequestedDepthSwapchainFormat == PF_DepthStencil);
 
-	// This is not a static swapchain
-	UnifiedCreateFlags |= TexCreate_Dynamic;
+	FXRSwapChainPtr& DepthSwapchain = PipelinedLayerStateRendering.DepthSwapchain;
+	const FRHITexture2D* const DepthSwapchainTexture = DepthSwapchain == nullptr ? nullptr : DepthSwapchain->GetTexture2DArray() ? DepthSwapchain->GetTexture2DArray() : DepthSwapchain->GetTexture2D();
+	if (DepthSwapchain == nullptr || DepthSwapchainTexture == nullptr || 
+		DepthSwapchainTexture->GetSizeX() != SizeX || DepthSwapchainTexture->GetSizeY() != SizeY)
+	{
+		// We're only creating a 1x target here, but we don't know whether it'll be the targeted texture
+		// or the resolve texture. Because of this, we unify the input flags.
+		ETextureCreateFlags UnifiedCreateFlags = TexCreate_DepthStencilTargetable | TexCreate_ShaderResource | TexCreate_InputAttachmentRead;
 
-	// We assume this could be used as a resolve target
-	UnifiedCreateFlags |= TexCreate_DepthStencilResolveTarget;
+		// This is not a static swapchain
+		UnifiedCreateFlags |= TexCreate_Dynamic;
 
-	// We can't use the depth swapchain w/o this flag, so client should always pass it in
-	ensure(EnumHasAllFlags(UnifiedCreateFlags, TexCreate_DepthStencilTargetable));
+		// We assume this could be used as a resolve target
+		UnifiedCreateFlags |= TexCreate_DepthStencilResolveTarget;
 
-	FXRSwapChainPtr& Swapchain = PipelinedLayerStateRendering.DepthSwapchain;
-	const FRHITexture2D* const DepthSwapchainTexture = Swapchain == nullptr ? nullptr : Swapchain->GetTexture2DArray() ? Swapchain->GetTexture2DArray() : Swapchain->GetTexture2D();
-	if (Swapchain == nullptr || DepthSwapchainTexture == nullptr || Format != LastRequestedDepthSwapchainFormat || DepthSwapchainTexture->GetSizeX() != SizeX || DepthSwapchainTexture->GetSizeY() != SizeY)
-	{
 		ensureMsgf(NumSamples == 1, TEXT("OpenXR supports MSAA swapchains, but engine logic expects the swapchain target to be 1x."));
+		constexpr uint32 NumSamplesExpected = 1;
+		constexpr uint32 NumMipsExpected = 1;
 
-		Swapchain = RenderBridge->CreateSwapchain(Session, PF_DepthStencil, SizeX, SizeY, bIsMobileMultiViewEnabled ? 2 : 1, FMath::Max(NumMips, 1u), NumSamples, UnifiedCreateFlags, FClearValueBinding::DepthFar);
-		if (!Swapchain)
+		DepthSwapchain = RenderBridge->CreateSwapchain(Session, PF_DepthStencil, SizeX, SizeY, bIsMobileMultiViewEnabled ? 2 : 1, NumMipsExpected, NumSamplesExpected, UnifiedCreateFlags, FClearValueBinding::DepthFar);
+		if (!DepthSwapchain)
 		{
-			return false;
+			return;
 		}
 
 		// image will be acquired next time we begin the rendering
 	}
 
 	bNeedReAllocatedDepth = false;
+}
+
+// TODO: in the future, we can rename the interface to GetDepthTexture because allocate could happen in AllocateRenderTargetTexture
+bool FOpenXRHMD::AllocateDepthTexture(uint32 Index, uint32 SizeX, uint32 SizeY, uint8 Format, uint32 NumMips, ETextureCreateFlags Flags, ETextureCreateFlags TargetableTextureFlags, FTexture2DRHIRef& OutTargetableTexture, FTexture2DRHIRef& OutShaderResourceTexture, uint32 NumSamples)
+{
+	check(IsInRenderingThread());
+
+	// FIXME: UE constantly calls this function even when there is no reason to reallocate the depth texture (see NeedReAllocateDepthTexture)
+	FReadScopeLock Lock(SessionHandleMutex);
+	if (!Session || !bDepthExtensionSupported)
+	{
+		return false;
+	}
+
+	const FXRSwapChainPtr& DepthSwapchain = PipelinedLayerStateRendering.DepthSwapchain;
+	if (DepthSwapchain == nullptr)
+	{
+		return false;
+	}
+
+	const ETextureCreateFlags UnifiedCreateFlags = Flags | TargetableTextureFlags;
+	ensure(EnumHasAllFlags(UnifiedCreateFlags, TexCreate_DepthStencilTargetable)); // We can't use the depth swapchain w/o this flag
+
+	const FRHITexture2D* const DepthSwapchainTexture = DepthSwapchain->GetTexture2DArray() ? DepthSwapchain->GetTexture2DArray() : DepthSwapchain->GetTexture2D();
+	const FRHITextureDesc& DepthSwapchainDesc = DepthSwapchainTexture->GetDesc();
+
+	// Sample count, mip count and size should be known at AllocateRenderTargetTexture time
+	// Format _could_ change, but we should know it (and can check for it in AllocateDepthTextureInternal)
+	// Flags might also change. We expect TexCreate_DepthStencilTargetable | TexCreate_ShaderResource | TexCreate_InputAttachmentRead from SceneTextures
+	check(EnumHasAllFlags(DepthSwapchainDesc.Flags, UnifiedCreateFlags));
+	check(DepthSwapchainDesc.Format == Format);
+	check(DepthSwapchainDesc.Extent.X == SizeX && DepthSwapchainDesc.Extent.Y == SizeY);
+	check(DepthSwapchainDesc.NumMips == FMath::Max(NumMips, 1u));
+	check(DepthSwapchainDesc.NumSamples == NumSamples);
 
-	OutTargetableTexture = OutShaderResourceTexture = (FTexture2DRHIRef&)Swapchain->GetTextureRef();
 	LastRequestedDepthSwapchainFormat = Format;
 
+	OutTargetableTexture = OutShaderResourceTexture = (FTexture2DRHIRef&)PipelinedLayerStateRendering.DepthSwapchain->GetTextureRef();
+
 	return true;
 }
 
@@ -2695,10 +2815,24 @@ bool FOpenXRHMD::HDRGetMetaDataForStereo(EDisplayOutputFormat& OutDisplayOutputF
 	return RenderBridge->HDRGetMetaDataForStereo(OutDisplayOutputFormat, OutDisplayColorGamut, OutbHDRSupported);
 }
 
+float FOpenXRHMD::GetPixelDenity() const
+{
+	const FPipelinedFrameState& PipelineState = GetPipelinedFrameStateForThread();
+	return PipelineState.PixelDensity;
+}
+
 void FOpenXRHMD::SetPixelDensity(const float NewDensity)
 {
 	check(IsInGameThread());
 	PipelinedFrameStateGame.PixelDensity = FMath::Min(NewDensity, RuntimePixelDensityMax);
+
+	// We have to update the RT state because the new swapchain will be allocated (FSceneViewport::InitDynamicRHI + AllocateRenderTargetTexture)
+	// before we call OnBeginRendering_GameThread.
+	ENQUEUE_RENDER_COMMAND(UpdatePixelDensity)(
+		[this, PixelDensity = PipelinedFrameStateGame.PixelDensity](FRHICommandListImmediate&)
+		{
+			PipelinedFrameStateRendering.PixelDensity = PixelDensity;
+		});
 }
 
 FIntPoint FOpenXRHMD::GetIdealRenderTargetSize() const
@@ -2715,8 +2849,7 @@ FIntPoint FOpenXRHMD::GetIdealRenderTargetSize() const
 			: Size.X + (int)Config.recommendedImageRectWidth;
 		Size.Y = FMath::Max(Size.Y, (int)Config.recommendedImageRectHeight);
 
-		// We always prefer the nearest multiple of 4 for our buffer sizes. Make sure we round up here,
-		// so we're consistent with the rest of the engine in creating our buffers.
+		// Make sure we quantize in order to be consistent with the rest of the engine in creating our buffers.
 		QuantizeSceneBufferSize(Size, Size);
 	}
 
diff --git a/Engine/Plugins/Runtime/OpenXR/Source/OpenXRHMD/Private/OpenXRHMD.h b/Engine/Plugins/Runtime/OpenXR/Source/OpenXRHMD/Private/OpenXRHMD.h
--- a/Engine/Plugins/Runtime/OpenXR/Source/OpenXRHMD/Private/OpenXRHMD.h
+++ b/Engine/Plugins/Runtime/OpenXR/Source/OpenXRHMD/Private/OpenXRHMD.h
@@ -232,12 +232,15 @@ protected:
 	void UpdateDeviceLocations(bool bUpdateOpenXRExtensionPlugins);
 	void EnumerateViews(FPipelinedFrameState& PipelineState);
 	void LocateViews(FPipelinedFrameState& PipelinedState, bool ResizeViewsArray = false);
+	bool IsViewManagedByPlugin(int32 ViewIndex) const;
 
 	void CopyTexture_RenderThread(FRHICommandListImmediate& RHICmdList, FRHITexture2D* SrcTexture, FIntRect SrcRect, FRHITexture2D* DstTexture, FIntRect DstRect, 
 								  bool bClearBlack, bool bNoAlpha, ERenderTargetActions RTAction, ERHIAccess FinalDstAccess) const;
 	
 	void CopyTexture_RenderThread(FRHICommandListImmediate& RHICmdList, FRHITexture2D* SrcTexture, FIntRect SrcRect, const FXRSwapChainPtr& DstSwapChain, FIntRect DstRect, bool bClearBlack, bool bNoAlpha) const;
 
+	void AllocateDepthTextureInternal(uint32 Index, uint32 SizeX, uint32 SizeY, uint32 NumSamples);
+
 	// Used with FCoreDelegates
 	void VRHeadsetRecenterDelegate();
 
@@ -256,7 +259,7 @@ public:
 	virtual bool GetHMDMonitorInfo(MonitorInfo&) override;
 	virtual void GetFieldOfView(float& OutHFOVInDegrees, float& OutVFOVInDegrees) const override;
 	virtual bool IsChromaAbCorrectionEnabled() const override;
-	//virtual float GetPixelDenity() const override;
+	virtual float GetPixelDenity() const override;
 	virtual void SetPixelDensity(const float NewDensity) override;
 	virtual FIntPoint GetIdealRenderTargetSize() const override;
 	virtual bool GetHMDDistortionEnabled(EShadingPath ShadingPath) const override { return false; }
@@ -293,12 +296,14 @@ public:
 	virtual void BeginRenderViewFamily(FSceneViewFamily& InViewFamily) override;
 	virtual void PreRenderView_RenderThread(FRDGBuilder& GraphBuilder, FSceneView& InView) override;
 	virtual void PreRenderViewFamily_RenderThread(FRDGBuilder& GraphBuilder, FSceneViewFamily& InViewFamily) override;
+	virtual void PostRenderViewFamily_RenderThread(FRDGBuilder& GraphBuilder, FSceneViewFamily& InViewFamily) override;
 
 	/** FHMDSceneViewExtension interface */
-	virtual bool IsActiveThisFrame_Internal(const FSceneViewExtensionContext& Context) const;
+	virtual bool IsActiveThisFrame_Internal(const FSceneViewExtensionContext& Context) const override;
 
 	/** IStereoRenderTargetManager */
 	virtual bool ShouldUseSeparateRenderTarget() const override { return IsStereoEnabled(); }
+	virtual void CalculateRenderTargetSize(const FViewport& Viewport, uint32& InOutSizeX, uint32& InOutSizeY) override;
 	virtual bool AllocateRenderTargetTexture(uint32 Index, uint32 SizeX, uint32 SizeY, uint8 Format, uint32 NumMips, ETextureCreateFlags Flags, ETextureCreateFlags TargetableTextureFlags, FTexture2DRHIRef& OutTargetableTexture, FTexture2DRHIRef& OutShaderResourceTexture, uint32 NumSamples = 1) override;
 	virtual bool NeedReAllocateDepthTexture(const TRefCountPtr<IPooledRenderTarget>& DepthTarget) override final { return bNeedReAllocatedDepth; }
 	virtual bool AllocateDepthTexture(uint32 Index, uint32 SizeX, uint32 SizeY, uint8 Format, uint32 NumMips, ETextureCreateFlags InTexFlags, ETextureCreateFlags TargetableTextureFlags, FTexture2DRHIRef& OutTargetableTexture, FTexture2DRHIRef& OutShaderResourceTexture, uint32 NumSamples = 1) override final;
diff --git a/Engine/Source/Runtime/HeadMountedDisplay/Private/XRRenderTargetManager.cpp b/Engine/Source/Runtime/HeadMountedDisplay/Private/XRRenderTargetManager.cpp
--- a/Engine/Source/Runtime/HeadMountedDisplay/Private/XRRenderTargetManager.cpp
+++ b/Engine/Source/Runtime/HeadMountedDisplay/Private/XRRenderTargetManager.cpp
@@ -20,25 +20,8 @@ void FXRRenderTargetManager::CalculateRenderTargetSize(const class FViewport& Vi
 		{
 			const FIntPoint IdealRenderTargetSize = HMDDevice->GetIdealRenderTargetSize();
 			const float PixelDensity = HMDDevice->GetPixelDenity();
-
-			FIntPoint DensityAdjustedTargetSize =
-			{
-				FMath::CeilToInt(IdealRenderTargetSize.X * PixelDensity),
-				FMath::CeilToInt(IdealRenderTargetSize.Y * PixelDensity)
-			};
-
-			// We need a custom quantized width here because if we have an atlased texture, each half needs to be aligned.
-			// We could modify or overload QuantizeSceneBufferSize, but this is the only call point that needs to fix up 
-			// the pixel-density-adjusted width. Strata requires DivBy8, hence aligning to 16.
-			// TODO: Would be nice if we could offload the alignment to QSBS by passing in number of atlased views
-			constexpr uint32 Mask16 = ~(16 - 1);
-			DensityAdjustedTargetSize.X = (DensityAdjustedTargetSize.X + (16 - 1)) & Mask16;
-
-			QuantizeSceneBufferSize(DensityAdjustedTargetSize, DensityAdjustedTargetSize);
-
-			InOutSizeX = DensityAdjustedTargetSize.X;
-			InOutSizeY = DensityAdjustedTargetSize.Y;
-
+			InOutSizeX = FMath::CeilToInt(IdealRenderTargetSize.X * PixelDensity);
+			InOutSizeY = FMath::CeilToInt(IdealRenderTargetSize.Y * PixelDensity);
 			check(InOutSizeX != 0 && InOutSizeY != 0);
 		}
 	}
