OpenXRAR - avoid assert in GetDisplayTime when the openxr session is not synchronized.
-Checking IsFocused before trying to get time.  xrLocateSpace is only supposed to work when in the focused state.  Our pins will retain their last known position.  During pin creation they won't locate at all which should result in them maintaining their unreal world position.  Both of these are good fallbacks.

diff --git a/Engine/Plugins/Runtime/OpenXR/Source/OpenXRAR/Private/OpenXRAR.cpp b/Engine/Plugins/Runtime/OpenXR/Source/OpenXRAR/Private/OpenXRAR.cpp
--- a/Engine/Plugins/Runtime/OpenXR/Source/OpenXRAR/Private/OpenXRAR.cpp
+++ b/Engine/Plugins/Runtime/OpenXR/Source/OpenXRAR/Private/OpenXRAR.cpp
@@ -246,7 +246,13 @@ UARPin* FOpenXRARSystem::OnPinComponent(USceneComponent* ComponentToPin, const F
 		if (CustomAnchorSupport != nullptr)
 		{
 			XrSession Session = TrackingSystem->GetSession();
-			XrTime DisplayTime = TrackingSystem->GetDisplayTime();
+			XrTime DisplayTime = 0;
+			if (TrackingSystem->IsFocused())
+			{
+				// If the tracking system is not focused we will try to pin with time zero.
+				// The locate should fail, but we will still setup the pin so it will work when the tracking system is working.
+				DisplayTime = TrackingSystem->GetDisplayTime();
+			}
 			XrSpace TrackingSpace = TrackingSystem->GetTrackingSpace();
 			float WorldToMetersScale = TrackingSystem->GetWorldToMetersScale();
 			if (!CustomAnchorSupport->OnPinComponent(NewPin, Session, TrackingSpace, DisplayTime, WorldToMetersScale))
@@ -286,13 +292,16 @@ void FOpenXRARSystem::UpdateAnchors()
 
 	if (CustomAnchorSupport != nullptr)
 	{
-		XrSession Session = TrackingSystem->GetSession();
-		XrTime DisplayTime = TrackingSystem->GetDisplayTime();
-		XrSpace TrackingSpace = TrackingSystem->GetTrackingSpace();
-		float WorldToMetersScale = TrackingSystem->GetWorldToMetersScale();
-		for (UARPin* Pin : Pins)
+		if (TrackingSystem->IsFocused())
 		{
-			CustomAnchorSupport->OnUpdatePin(Pin, Session, TrackingSpace, DisplayTime, WorldToMetersScale);
+			XrSession Session = TrackingSystem->GetSession();
+			XrTime DisplayTime = TrackingSystem->GetDisplayTime();
+			XrSpace TrackingSpace = TrackingSystem->GetTrackingSpace();
+			float WorldToMetersScale = TrackingSystem->GetWorldToMetersScale();
+			for (UARPin* Pin : Pins)
+			{
+				CustomAnchorSupport->OnUpdatePin(Pin, Session, TrackingSpace, DisplayTime, WorldToMetersScale);
+			}
 		}
 	}
 }
