OpenXR Enhanced Input refactor
* Add Start XR Session blueprint node that allows changing the PMIC at runtime
* Make the controller action set persistent to reduce complexity
* Move interaction profile array to a local member used to build suggested bindings

diff --git a/Engine/Plugins/Runtime/OpenXR/Source/OpenXRHMD/Public/IOpenXRExtensionPlugin.h b/Engine/Plugins/Runtime/OpenXR/Source/OpenXRHMD/Public/IOpenXRExtensionPlugin.h
--- a/Engine/Plugins/Runtime/OpenXR/Source/OpenXRHMD/Public/IOpenXRExtensionPlugin.h
+++ b/Engine/Plugins/Runtime/OpenXR/Source/OpenXRHMD/Public/IOpenXRExtensionPlugin.h
@@ -255,12 +255,16 @@ public:
 	{
 	}
 
+	UE_DEPRECATED(5.1, "Functionality moved to AttachActionSets().")
+	virtual void AddActionSets(TArray<XrActiveActionSet>& OutActionSets)
+	{
+	}
+
 	/**
 	* Add any action sets provided by the plugin to be attached as active to the session
 	* This allows a plugin to manage a custom actionset that will be active in xrSyncActions
 	*/
-	UE_DEPRECATED(5.1, "Use AddActions to create attached action sets, along with GetActiveActionSetsForSync.")
-	virtual void AddActionSets(TArray<XrActiveActionSet>& OutActionSets)
+	virtual void AttachActionSets(TSet<XrActionSet>& OutActionSets)
 	{
 	}
 
diff --git a/Engine/Plugins/Runtime/OpenXR/Source/OpenXRInput/Private/OpenXRInput.cpp b/Engine/Plugins/Runtime/OpenXR/Source/OpenXRInput/Private/OpenXRInput.cpp
--- a/Engine/Plugins/Runtime/OpenXR/Source/OpenXRInput/Private/OpenXRInput.cpp
+++ b/Engine/Plugins/Runtime/OpenXR/Source/OpenXRInput/Private/OpenXRInput.cpp
@@ -243,7 +243,8 @@ FOpenXRInputPlugin::FInteractionProfile::FInteractionProfile(XrPath InProfile, b
 
 FOpenXRInputPlugin::FOpenXRInput::FOpenXRInput(FOpenXRHMD* HMD)
 	: OpenXRHMD(HMD)
-	, Profiles()
+	, Instance(XR_NULL_HANDLE)
+	, ControllerActionSet()
 	, ActionSets()
 	, PluginActionSets()
 	, SubactionPaths()
@@ -251,10 +252,8 @@ FOpenXRInputPlugin::FOpenXRInput::FOpenXRInput(FOpenXRHMD* HMD)
 	, EnhancedActions()
 	, Controllers()
 	, MotionSourceToControllerHandMap()
-	, ControllerSet()
-	, LegacySet()
 	, MappableInputConfig(nullptr)
-	, bActionsBound(false)
+	, bActionsAttached(false)
 	, bDirectionalBindingSupported(false)
 	, MessageHandler(new FGenericApplicationMessageHandler())
 {
@@ -263,6 +262,7 @@ FOpenXRInputPlugin::FOpenXRInput::FOpenXRInput(FOpenXRHMD* HMD)
 	// If there is no HMD then this module is not active, but it still needs to exist so we can EnumerateMotionSources from it.
 	if (OpenXRHMD)
 	{
+		Instance = OpenXRHMD->GetInstance();
 		bDirectionalBindingSupported = OpenXRHMD->IsExtensionEnabled("XR_EXT_dpad_binding");
 
 		// Note: AnyHand needs special handling because it tries left then falls back to right in each call.
@@ -278,7 +278,32 @@ FOpenXRInputPlugin::FOpenXRInput::FOpenXRInput(FOpenXRHMD* HMD)
 		MotionSourceToControllerHandMap.Add(TEXT("EControllerHand::Right"), EControllerHand::Right);
 		MotionSourceToControllerHandMap.Add(TEXT("EControllerHand::AnyHand"), EControllerHand::AnyHand);
 
-		BuildActions();
+		// Create an engine action set for pose input and haptic output
+		ControllerActionSet = MakeUnique<FOpenXRActionSet>(Instance, "controllers", "Controllers", 0);
+
+		XrPath LeftHand = GetPath(Instance, "/user/hand/left");
+		XrPath RightHand = GetPath(Instance, "/user/hand/right");
+		XrPath Head = GetPath(Instance, "/user/head");
+
+		// Controller poses
+		Controllers.Add(EControllerHand::Left, FOpenXRController(ControllerActionSet->Handle, LeftHand, "Left Controller"));
+		Controllers.Add(EControllerHand::Right, FOpenXRController(ControllerActionSet->Handle, RightHand, "Right Controller"));
+		Controllers.Add(EControllerHand::HMD, FOpenXRController(ControllerActionSet->Handle, Head, "HMD"));
+
+		// Generate a list of the sub-action paths so we can query the left/right hand individually
+		SubactionPaths.Add(LeftHand);
+		SubactionPaths.Add(RightHand);
+
+		// Make OpenXRHMD aware of the controller action spaces
+		Controllers[EControllerHand::Left].AddActionDevices(OpenXRHMD);
+		Controllers[EControllerHand::Right].AddActionDevices(OpenXRHMD);
+
+		// Attempt to load the default input config from the OpenXR input settings.
+		UOpenXRInputSettings* InputSettings = GetMutableDefault<UOpenXRInputSettings>();
+		if (InputSettings && InputSettings->MappableInputConfig.IsValid())
+		{
+			MappableInputConfig = (UPlayerMappableInputConfig*)InputSettings->MappableInputConfig.TryLoad();
+		}
 	}
 }
 
@@ -329,22 +354,25 @@ bool FOpenXRInputPlugin::FOpenXRInput::IsOpenXRInputSupportedMotionSource(const
 FOpenXRInputPlugin::FOpenXRInput::~FOpenXRInput()
 {
 	DestroyActions();
+
+	if (ControllerActionSet)
+	{
+		xrDestroyActionSet(ControllerActionSet->Handle);
+	}
 }
 
-void FOpenXRInputPlugin::FOpenXRInput::BuildActions()
+bool FOpenXRInputPlugin::FOpenXRInput::BuildActions(XrSession Session)
 {
-	if ((bActionsBound) || (OpenXRHMD == nullptr))
+	if (bActionsAttached)
 	{
-		return;
+		return false;
 	}
 
-	XrInstance Instance = OpenXRHMD->GetInstance();
-	check(Instance);
 	DestroyActions();
 
-	// Generate a map of all supported interaction profiles
-	XrPath SimpleControllerPath = GetPath(Instance, "/interaction_profiles/khr/simple_controller");
-	Profiles.Add("SimpleController", FInteractionProfile(SimpleControllerPath, true));
+	// Generate a map of all supported interaction profiles to store suggested bindings
+	TMap<FString, FInteractionProfile> Profiles;
+	Profiles.Add("SimpleController", FInteractionProfile(GetPath(Instance, "/interaction_profiles/khr/simple_controller"), true));
 	Profiles.Add("Vive", FInteractionProfile(GetPath(Instance, "/interaction_profiles/htc/vive_controller"), true));
 	Profiles.Add("MixedReality", FInteractionProfile(GetPath(Instance, "/interaction_profiles/microsoft/motion_controller"), true));
 	Profiles.Add("OculusGo", FInteractionProfile(GetPath(Instance, "/interaction_profiles/oculus/go_controller"), false));
@@ -363,34 +391,13 @@ void FOpenXRInputPlugin::FOpenXRInput::BuildActions()
 		}
 	}
 
-	// Create an engine action set for pose input and haptic output
-	FOpenXRActionSet ActionSet(Instance, "controllers", "Controllers", 0);
-	ControllerSet.actionSet = ActionSet.Handle;
-	ControllerSet.subactionPath = XR_NULL_PATH;
-
-	XrPath LeftHand = GetPath(Instance, "/user/hand/left");
-	XrPath RightHand = GetPath(Instance, "/user/hand/right");
-	XrPath HMD = GetPath(Instance, "/user/head");
-
-	// Controller poses
-	OpenXRHMD->ResetActionDevices();
-	Controllers.Add(EControllerHand::Left, FOpenXRController(ActionSet.Handle, LeftHand, "Left Controller"));
-	Controllers.Add(EControllerHand::Right, FOpenXRController(ActionSet.Handle, RightHand, "Right Controller"));
-	Controllers.Add(EControllerHand::HMD, FOpenXRController(ActionSet.Handle, HMD, "HMD"));
-	ActionSets.Emplace(MoveTemp(ActionSet));
-
-	// Generate a list of the sub-action paths so we can query the left/right hand individually
-	SubactionPaths.Add(LeftHand);
-	SubactionPaths.Add(RightHand);
-
-	UOpenXRInputSettings* InputSettings = GetMutableDefault<UOpenXRInputSettings>();
-	if (InputSettings && InputSettings->bUseEnhancedInput)
+	if (MappableInputConfig)
 	{
-		BuildEnhancedActions(InputSettings->MappableInputConfig);
+		BuildEnhancedActions(Profiles, MappableInputConfig);
 	}
 	else
 	{
-		BuildLegacyActions();
+		BuildLegacyActions(Profiles);
 	}
 
 	for (TPair<FString, FInteractionProfile>& Pair : Profiles)
@@ -399,7 +406,7 @@ void FOpenXRInputPlugin::FOpenXRInput::BuildActions()
 
 		// Only suggest interaction profile bindings if the developer has provided bindings for them
 		// An exception is made for the Simple Controller Profile which is always bound as a fallback
-		if (Profile.Bindings.Num() > 0 || Profile.Path == SimpleControllerPath)
+		if (Profile.Bindings.Num() > 0)
 		{
 			// Add the bindings for the controller pose and haptics
 			Profile.Bindings.Add(XrActionSuggestedBinding{
@@ -436,20 +443,54 @@ void FOpenXRInputPlugin::FOpenXRInput::BuildActions()
 		}
 	}
 
-	Controllers[EControllerHand::Left].AddActionDevices(OpenXRHMD);
-	Controllers[EControllerHand::Right].AddActionDevices(OpenXRHMD);
+	// Bind the project action sets
+	TSet<XrActionSet> AttachSet;
+	for (auto&& ActionSet : ActionSets)
+		AttachSet.Add(ActionSet.Handle);
+
+	// Bind the controller action set
+	if (ControllerActionSet)
+	{
+		AttachSet.Add(ControllerActionSet->Handle);
+	}
+
+	// Bind plugin action sets exposed through a query
+	for (IOpenXRExtensionPlugin* Plugin : OpenXRHMD->GetExtensionPlugins())
+	{
+		TArray<XrActiveActionSet> PluginAttachArray_Deprecated;
+		PRAGMA_DISABLE_DEPRECATION_WARNINGS;
+		// TODO?: Log deprecation warning at runtime, since overridden deprecated interface methods don't warn at compile time?
+		Plugin->AddActionSets(PluginAttachArray_Deprecated);
+		PRAGMA_ENABLE_DEPRECATION_WARNINGS;
+		for (const XrActiveActionSet& ActiveSet : PluginAttachArray_Deprecated)
+		{
+			AttachSet.Add(ActiveSet.actionSet);
+		}
+
+		TSet<XrActionSet> PluginAttachSet;
+		Plugin->AttachActionSets(PluginAttachSet);
+		AttachSet.Append(PluginAttachSet);
+	}
+
+	// Bind the plugin action sets created using the deprecated API
+	for (auto&& ActionSet : PluginActionSets)
+		AttachSet.Add(ActionSet.Handle);
+
+	TArray<XrActionSet> AttachArray = AttachSet.Array();
+	XrSessionActionSetsAttachInfo SessionActionSetsAttachInfo;
+	SessionActionSetsAttachInfo.type = XR_TYPE_SESSION_ACTION_SETS_ATTACH_INFO;
+	SessionActionSetsAttachInfo.next = nullptr;
+	SessionActionSetsAttachInfo.countActionSets = AttachArray.Num();
+	SessionActionSetsAttachInfo.actionSets = AttachArray.GetData();
+	bActionsAttached = XR_ENSURE(xrAttachSessionActionSets(Session, &SessionActionSetsAttachInfo));
+
+	return bActionsAttached;
 }
 
-void FOpenXRInputPlugin::FOpenXRInput::BuildLegacyActions()
+void FOpenXRInputPlugin::FOpenXRInput::BuildLegacyActions(TMap<FString, FInteractionProfile>& Profiles)
 {
-	XrInstance Instance = OpenXRHMD->GetInstance();
-	check(Instance);
-
-	// Create an engine action set for pose input and haptic output
+	// Create an engine action set for legacy actions
 	FOpenXRActionSet ActionSet(Instance, "ue", "Unreal Engine", 0);
-	LegacySet.actionSet = ActionSet.Handle;
-	LegacySet.subactionPath = XR_NULL_PATH;
-
 	auto InputSettings = GetMutableDefault<UInputSettings>();
 	if (InputSettings != nullptr)
 	{
@@ -462,7 +503,7 @@ void FOpenXRInputPlugin::FOpenXRInput::BuildLegacyActions()
 			InputSettings->GetActionMappingByName(ActionName, Mappings);
 
 			// If the developer didn't suggest any XR bindings for this action we won't expose it to the runtime
-			if (SuggestBindings(Instance, Action, Mappings) > 0)
+			if (SuggestBindings(Profiles, Action, Mappings) > 0)
 			{
 				LegacyActions.Add(Action);
 			}
@@ -481,7 +522,7 @@ void FOpenXRInputPlugin::FOpenXRInput::BuildLegacyActions()
 			InputSettings->GetAxisMappingByName(AxisName, Mappings);
 
 			// If the developer didn't suggest any XR bindings for this action we won't expose it to the runtime
-			if (SuggestBindings(Instance, Action, Mappings) > 0)
+			if (SuggestBindings(Profiles, Action, Mappings) > 0)
 			{
 				LegacyActions.Add(Action);
 			}
@@ -513,7 +554,7 @@ void FOpenXRInputPlugin::FOpenXRInput::BuildLegacyActions()
 	for (IOpenXRExtensionPlugin* Plugin : OpenXRHMD->GetExtensionPlugins())
 	{
 		IOpenXRExtensionPlugin::FCreateActionSetFunc CreateActionSetFunc =
-			[this, Instance]
+			[this]
 			(const IOpenXRExtensionPlugin::FActionSetParams& Params) -> XrActionSet
 			{
 				FOpenXRActionSet ActionSet(Instance, Params.Name, Params.LocalizedName.ToString(), Params.Priority);
@@ -523,7 +564,7 @@ void FOpenXRInputPlugin::FOpenXRInput::BuildLegacyActions()
 			};
 
 		IOpenXRExtensionPlugin::FCreateActionFunc CreateActionFunc =
-			[this, Instance, ActionSet]
+			[this, &Profiles, ActionSet]
 			(const IOpenXRExtensionPlugin::FActionParams& Params) -> XrAction
 			{
 				// TODO: Use Params.LocalizedName
@@ -532,31 +573,34 @@ void FOpenXRInputPlugin::FOpenXRInput::BuildLegacyActions()
 				const XrAction ReturnHandle = Action.Handle;
 				for (const FKey& Key : Params.SuggestedBindings)
 				{
-					SuggestBindingForKey(Instance, Action, Key);
+					SuggestBindingForKey(Profiles, Action, Key);
 				}
 				LegacyActions.Emplace(MoveTemp(Action));
 				return ReturnHandle;
 			};
 
+		PRAGMA_DISABLE_DEPRECATION_WARNINGS;
 		Plugin->AddActions(Instance, CreateActionSetFunc, CreateActionFunc);
+		PRAGMA_ENABLE_DEPRECATION_WARNINGS;
 	}
 
-	ActionSets.Emplace(MoveTemp(ActionSet));
+	if (!LegacyActions.IsEmpty())
+	{
+		UE_LOG(LogHMD, Warning, TEXT("Action/Axis mappings from the Input configuration are deprecated for XR in favor of Enhanced Input."));
+
+		ActionSets.Emplace(MoveTemp(ActionSet));
+	}
 }
 
-void FOpenXRInputPlugin::FOpenXRInput::BuildEnhancedActions(const FSoftObjectPath& MappableInputConfigPath)
+void FOpenXRInputPlugin::FOpenXRInput::BuildEnhancedActions(TMap<FString, FInteractionProfile>& Profiles, TObjectPtr<UPlayerMappableInputConfig> InputConfig)
 {
-	XrInstance Instance = OpenXRHMD->GetInstance();
-	check(Instance);
-
-	MappableInputConfig = (UPlayerMappableInputConfig*)MappableInputConfigPath.TryLoad();
-	if (!MappableInputConfig)
+	if (bActionsAttached)
 	{
+		// Don't rebuild actions while another input config has already been attached
 		return;
 	}
-	MappableInputConfig->AddToRoot();
 
-	for (const TPair<TObjectPtr<UInputMappingContext>, int32> MappingContext : MappableInputConfig->GetMappingContexts())
+	for (const TPair<TObjectPtr<UInputMappingContext>, int32> MappingContext : InputConfig->GetMappingContexts())
 	{
 		FOpenXRActionSet ActionSet(Instance, MappingContext.Key->GetFName(), MappingContext.Key->ContextDescription.ToString(), MappingContext.Value, MappingContext.Key);
 		TMap<FName, int32> ActionMap;
@@ -585,7 +629,7 @@ void FOpenXRInputPlugin::FOpenXRInput::BuildEnhancedActions(const FSoftObjectPat
 				ActionIndex = EnhancedActions.Emplace(ActionSet.Handle, ActionType, ActionName, LocalizedName, SubactionPaths, Mapping.Action);
 			}
 
-			SuggestBindingForKey(Instance, EnhancedActions[ActionIndex], Mapping.Key, Mapping.Modifiers, Mapping.Triggers);
+			SuggestBindingForKey(Profiles, EnhancedActions[ActionIndex], Mapping.Key, Mapping.Modifiers, Mapping.Triggers);
 		}
 		ActionSets.Emplace(MoveTemp(ActionSet));
 	}
@@ -599,30 +643,21 @@ void FOpenXRInputPlugin::FOpenXRInput::DestroyActions()
 		xrDestroyActionSet(ActionSet.Handle);
 	}
 
-	// Clear the active controller and legacy action set
-	ControllerSet.actionSet = XR_NULL_HANDLE;
-	ControllerSet.subactionPath = XR_NULL_PATH;
-	LegacySet.actionSet = XR_NULL_HANDLE;
-	LegacySet.subactionPath = XR_NULL_PATH;
-
-	Profiles.Reset();
 	LegacyActions.Reset();
 	EnhancedActions.Reset();
-	Controllers.Reset();
-	SubactionPaths.Reset();
 	ActionSets.Reset();
 	PluginActionSets.Reset();
 }
 
 template<typename T>
-int32 FOpenXRInputPlugin::FOpenXRInput::SuggestBindings(XrInstance Instance, FOpenXRAction& Action, const TArray<T>& Mappings)
+int32 FOpenXRInputPlugin::FOpenXRInput::SuggestBindings(TMap<FString, FInteractionProfile>& Profiles, FOpenXRAction& Action, const TArray<T>& Mappings)
 {
 	int32 SuggestedBindings = 0;
 
 	// Add suggested bindings for every mapping
 	for (const T& InputKey : Mappings)
 	{
-		if (SuggestBindingForKey(Instance, Action, InputKey.Key))
+		if (SuggestBindingForKey(Profiles, Action, InputKey.Key))
 		{
 			++SuggestedBindings;
 		}
@@ -631,7 +666,7 @@ int32 FOpenXRInputPlugin::FOpenXRInput::SuggestBindings(XrInstance Instance, FOp
 	return SuggestedBindings;
 }
 
-bool FOpenXRInputPlugin::FOpenXRInput::SuggestBindingForKey(XrInstance Instance, FOpenXRAction& Action, const FKey& InFKey, const TArray<UInputModifier*>& Modifiers, const TArray<UInputTrigger*>& Triggers)
+bool FOpenXRInputPlugin::FOpenXRInput::SuggestBindingForKey(TMap<FString, FInteractionProfile>& Profiles, FOpenXRAction& Action, const FKey& InFKey, const TArray<UInputModifier*>& Modifiers, const TArray<UInputTrigger*>& Triggers)
 {
 	// Key names that are parseable into an OpenXR path have exactly 4 tokens
 	TArray<FString> Tokens;
@@ -728,11 +763,11 @@ bool FOpenXRInputPlugin::FOpenXRInput::SuggestBindingForKey(XrInstance Instance,
 	return true;
 }
 
-bool FOpenXRInputPlugin::FOpenXRInput::SuggestBindingForKey(XrInstance Instance, FOpenXRAction& Action, const FKey& InFKey)
+bool FOpenXRInputPlugin::FOpenXRInput::SuggestBindingForKey(TMap<FString, FInteractionProfile>& Profiles, FOpenXRAction& Action, const FKey& InFKey)
 {
 	TArray<UInputModifier*> Modifiers;
 	TArray<UInputTrigger*> Triggers;
-	return SuggestBindingForKey(Instance, Action, InFKey, Modifiers, Triggers);
+	return SuggestBindingForKey(Profiles, Action, InFKey, Modifiers, Triggers);
 }
 
 void FOpenXRInputPlugin::FOpenXRInput::Tick(float DeltaTime)
@@ -746,75 +781,28 @@ void FOpenXRInputPlugin::FOpenXRInput::Tick(float DeltaTime)
 	XrSession Session = OpenXRHMD->GetSession();
 	if (Session != XR_NULL_HANDLE)
 	{
-		if (!bActionsBound)
+		// If we're using legacy actions we attach them to the session immediately.
+		// Otherwise we wait for a blueprint to attach enhanced input mappings.
+		if (!bActionsAttached)
 		{
-			BuildActions();
-
-			// Bind engine action sets
-			TArray<XrActionSet> BindActionSets;
-			for (auto && BindActionSet : ActionSets)
-				BindActionSets.Add(BindActionSet.Handle);
-
-			// Bind plugin action sets exposed via deprecated method
-			for (IOpenXRExtensionPlugin* Plugin : OpenXRHMD->GetExtensionPlugins())
-			{
-				TArray<XrActiveActionSet> PluginAttachSets_Deprecated;
-				PRAGMA_DISABLE_DEPRECATION_WARNINGS;
-				Plugin->AddActionSets(PluginAttachSets_Deprecated);
-				PRAGMA_ENABLE_DEPRECATION_WARNINGS;
-				for (const XrActiveActionSet& ActiveSet : PluginAttachSets_Deprecated)
-				{
-					// TODO?: Log deprecation warning at runtime, since overridden deprecated interface methods don't warn at compile time?
-					BindActionSets.Add(ActiveSet.actionSet);
-				}
-			}
-
-			for (auto&& BindActionSet : PluginActionSets)
-				BindActionSets.Add(BindActionSet.Handle);
-
-			XrSessionActionSetsAttachInfo SessionActionSetsAttachInfo;
-			SessionActionSetsAttachInfo.type = XR_TYPE_SESSION_ACTION_SETS_ATTACH_INFO;
-			SessionActionSetsAttachInfo.next = nullptr;
-			SessionActionSetsAttachInfo.countActionSets = BindActionSets.Num();
-			SessionActionSetsAttachInfo.actionSets = BindActionSets.GetData();
-			XR_ENSURE(xrAttachSessionActionSets(Session, &SessionActionSetsAttachInfo));
-
-			bActionsBound = true;
+			BuildActions(Session);
 		}
+
+		SyncActions(Session);
 	}
-	else if (bActionsBound)
+	else if (bActionsAttached)
 	{
 		// If the session shut down, clean up.
-		bActionsBound = false;
-
-		// Remove the input config from the root
-		if (MappableInputConfig)
-		{
-			MappableInputConfig->RemoveFromRoot();
-			MappableInputConfig = nullptr;
-		}
+		bActionsAttached = false;
 	}
+}
 
-	if (OpenXRHMD->IsFocused())
+void FOpenXRInputPlugin::FOpenXRInput::SyncActions(XrSession Session)
+{
+	if (OpenXRHMD->IsFocused() && bActionsAttached)
 	{
-		XrActionsSyncInfo SyncInfo;
-		SyncInfo.type = XR_TYPE_ACTIONS_SYNC_INFO;
-		SyncInfo.next = nullptr;
-
-		TArray<XrActiveActionSet> ActiveSets;
-		for (IOpenXRExtensionPlugin* Plugin : OpenXRHMD->GetExtensionPlugins())
-		{
-			Plugin->GetActiveActionSetsForSync(ActiveSets);
-
-			// TODO?: Log deprecation warning at runtime, since overridden deprecated interface methods don't warn at compile time?
-			PRAGMA_DISABLE_DEPRECATION_WARNINGS;
-			Plugin->AddActionSets(ActiveSets);
-			PRAGMA_ENABLE_DEPRECATION_WARNINGS;
-
-			SyncInfo.next = Plugin->OnSyncActions(Session, SyncInfo.next);
-		}
-
-		IEnhancedInputModule::Get().GetLibrary()->ForEachSubsystem([this, &ActiveSets](IEnhancedInputSubsystemInterface* Subsystem)
+		TSet<XrActionSet> ActiveSet;
+		IEnhancedInputModule::Get().GetLibrary()->ForEachSubsystem([this, &ActiveSet](IEnhancedInputSubsystemInterface* Subsystem)
 			{
 				if (Subsystem)
 				{
@@ -822,27 +810,46 @@ void FOpenXRInputPlugin::FOpenXRInput::Tick(float DeltaTime)
 					{
 						if (ActionSet.Object && Subsystem->HasMappingContext(ActionSet.Object))
 						{
-							XrActiveActionSet ActiveSet;
-							ActiveSet.actionSet = ActionSet.Handle;
-							ActiveSet.subactionPath = XR_NULL_PATH;
-							ActiveSets.Add(ActiveSet);
+							ActiveSet.Add(ActionSet.Handle);
 						}
 					}
 				}
 			});
 
-		// The controller and legacy action sets are always active
-		if (ControllerSet.actionSet)
+		// If legacy actions are enabled, all action sets are always active
+		if (!LegacyActions.IsEmpty())
 		{
-			ActiveSets.Add(ControllerSet);
+			for (const FOpenXRActionSet& ActionSet : ActionSets)
+			{
+				ActiveSet.Add(ActionSet.Handle);
+			}
 		}
-		if (LegacySet.actionSet)
+
+		// The controller is always active
+		if (ControllerActionSet)
 		{
-			ActiveSets.Add(LegacySet);
+			ActiveSet.Add(ControllerActionSet->Handle);
 		}
 
-		SyncInfo.countActiveActionSets = ActiveSets.Num();
-		SyncInfo.activeActionSets = ActiveSets.GetData();
+		XrActionsSyncInfo SyncInfo;
+		SyncInfo.type = XR_TYPE_ACTIONS_SYNC_INFO;
+		SyncInfo.next = nullptr;
+
+		TArray<XrActiveActionSet> ActiveActionSets;
+		for (XrActionSet ActionSet : ActiveSet)
+		{
+			ActiveActionSets.Add(XrActiveActionSet{ ActionSet, XR_NULL_PATH });
+		}
+
+		for (IOpenXRExtensionPlugin* Plugin : OpenXRHMD->GetExtensionPlugins())
+		{
+			Plugin->GetActiveActionSetsForSync(ActiveActionSets);
+
+			SyncInfo.next = Plugin->OnSyncActions(Session, SyncInfo.next);
+		}
+
+		SyncInfo.countActiveActionSets = ActiveActionSets.Num();
+		SyncInfo.activeActionSets = ActiveActionSets.GetData();
 
 		XR_ENSURE(xrSyncActions(Session, &SyncInfo));
 
@@ -868,7 +875,7 @@ namespace OpenXRInputNamespace
 
 void FOpenXRInputPlugin::FOpenXRInput::SendControllerEvents()
 {
-	if (!bActionsBound || OpenXRHMD == nullptr)
+	if (!bActionsAttached || OpenXRHMD == nullptr)
 	{
 		return;
 	}
@@ -1096,7 +1103,7 @@ void FOpenXRInputPlugin::FOpenXRInput::SetChannelValues(int32 ControllerId, cons
 
 bool FOpenXRInputPlugin::FOpenXRInput::SupportsForceFeedback(int32 ControllerId)
 {
-	if (!bActionsBound || OpenXRHMD == nullptr)
+	if (!bActionsAttached || OpenXRHMD == nullptr)
 	{
 		return false;
 	}
@@ -1314,7 +1321,7 @@ void FOpenXRInputPlugin::FOpenXRInput::EnumerateSources(TArray<FMotionController
 // TODO: Refactor API to change the Hand type to EControllerHand
 void FOpenXRInputPlugin::FOpenXRInput::SetHapticFeedbackValues(int32 ControllerId, int32 Hand, const FHapticFeedbackValues& Values)
 {
-	if (!bActionsBound || OpenXRHMD == nullptr)
+	if (!bActionsAttached || OpenXRHMD == nullptr)
 	{
 		return;
 	}
@@ -1413,4 +1420,34 @@ float FOpenXRInputPlugin::FOpenXRInput::GetHapticAmplitudeScale() const
 	return 1.0f;
 }
 
+bool FOpenXRInputPlugin::FOpenXRInput::SetPlayerMappableInputConfig(TObjectPtr<class UPlayerMappableInputConfig> InputConfig)
+{
+	if (bActionsAttached)
+	{
+		UE_LOG(LogHMD, Error, TEXT("Attempted to attach an input config while one is already attached for the current session."));
+
+		return false;
+	}
+
+	if (MappableInputConfig == InputConfig)
+	{
+		return true;
+	}
+
+	// Remove the previous config from the scene root
+	if (MappableInputConfig)
+	{
+		MappableInputConfig->RemoveFromRoot();
+	}
+
+	// If an input config is set, then attach it to the scene root to prevent it from being unloaded.
+	if (InputConfig)
+	{
+		InputConfig->AddToRoot();
+	}
+
+	MappableInputConfig = InputConfig;
+	return true;
+}
+
 #undef LOCTEXT_NAMESPACE // "OpenXRInputPlugin"
diff --git a/Engine/Plugins/Runtime/OpenXR/Source/OpenXRInput/Private/OpenXRInput.h b/Engine/Plugins/Runtime/OpenXR/Source/OpenXRInput/Private/OpenXRInput.h
--- a/Engine/Plugins/Runtime/OpenXR/Source/OpenXRInput/Private/OpenXRInput.h
+++ b/Engine/Plugins/Runtime/OpenXR/Source/OpenXRInput/Private/OpenXRInput.h
@@ -122,6 +122,7 @@ public:
 		virtual bool GetControllerOrientationAndPosition(const int32 ControllerIndex, const EControllerHand DeviceHand, FRotator& OutOrientation, FVector& OutPosition, float WorldToMetersScale) const override { check(false); return false; }
 		virtual ETrackingStatus GetControllerTrackingStatus(const int32 ControllerIndex, const EControllerHand DeviceHand) const override { check(false); return ETrackingStatus::NotTracked; }
 		virtual void EnumerateSources(TArray<FMotionControllerSource>& SourcesOut) const override;
+		virtual bool SetPlayerMappableInputConfig(TObjectPtr<class UPlayerMappableInputConfig> InputConfig) override;
 
 		// IHapticDevice overrides
 		IHapticDevice* GetHapticDevice() override { return (IHapticDevice*)this; }
@@ -134,8 +135,9 @@ public:
 		static const XrDuration MaxFeedbackDuration = 2500000000; // 2.5 seconds
 
 		FOpenXRHMD* OpenXRHMD;
+		XrInstance Instance;
 
-		TMap<FString, FInteractionProfile> Profiles;
+		TUniquePtr<FOpenXRActionSet> ControllerActionSet;
 		TArray<FOpenXRActionSet> ActionSets;
 		TArray<FOpenXRActionSet> PluginActionSets;
 		TArray<XrPath> SubactionPaths;
@@ -143,28 +145,26 @@ public:
 		TMap<EControllerHand, FOpenXRController> Controllers;
 		TMap<FName, EControllerHand> MotionSourceToControllerHandMap;
 
-		XrActiveActionSet ControllerSet;
-		XrActiveActionSet LegacySet;
-
-		UPlayerMappableInputConfig* MappableInputConfig;
+		TObjectPtr<UPlayerMappableInputConfig> MappableInputConfig;
 
 		XrAction GetActionForMotionSource(FName MotionSource) const;
 		int32 GetDeviceIDForMotionSource(FName MotionSource) const;
 		XrPath GetUserPathForMotionSource(FName MotionSource) const;
 		bool IsOpenXRInputSupportedMotionSource(const FName MotionSource) const;
 
-		bool bActionsBound;
+		bool bActionsAttached;
 		bool bDirectionalBindingSupported;
 
-		void BuildActions();
-		void BuildLegacyActions();
-		void BuildEnhancedActions(const FSoftObjectPath& MappableInputConfigPath);
+		bool BuildActions(XrSession Session);
+		void SyncActions(XrSession Session);
+		void BuildLegacyActions(TMap<FString, FInteractionProfile>& Profiles);
+		void BuildEnhancedActions(TMap<FString, FInteractionProfile>& Profiles, TObjectPtr<UPlayerMappableInputConfig> MappableInputConfig);
 		void DestroyActions();
 
 		template<typename T>
-		int32 SuggestBindings(XrInstance Instance, FOpenXRAction& Action, const TArray<T>& Mappings);
-		bool SuggestBindingForKey(XrInstance Instance, FOpenXRAction& Action, const FKey& Key, const TArray<UInputModifier*>& Modifiers, const TArray<UInputTrigger*>& Triggers);
-		bool SuggestBindingForKey(XrInstance Instance, FOpenXRAction& Action, const FKey& Key);
+		int32 SuggestBindings(TMap<FString, FInteractionProfile>& Profiles, FOpenXRAction& Action, const TArray<T>& Mappings);
+		bool SuggestBindingForKey(TMap<FString, FInteractionProfile>& Profiles, FOpenXRAction& Action, const FKey& Key, const TArray<UInputModifier*>& Modifiers, const TArray<UInputTrigger*>& Triggers);
+		bool SuggestBindingForKey(TMap<FString, FInteractionProfile>& Profiles, FOpenXRAction& Action, const FKey& Key);
 
 		/** handler to send all messages to */
 		TSharedRef<FGenericApplicationMessageHandler> MessageHandler;
diff --git a/Engine/Plugins/Runtime/OpenXR/Source/OpenXRInput/Private/OpenXRInputFunctionLibrary.cpp b/Engine/Plugins/Runtime/OpenXR/Source/OpenXRInput/Private/OpenXRInputFunctionLibrary.cpp
new file mode 100644
--- /dev/null
+++ b/Engine/Plugins/Runtime/OpenXR/Source/OpenXRInput/Private/OpenXRInputFunctionLibrary.cpp
@@ -0,0 +1,48 @@
+// Copyright Epic Games, Inc. All Rights Reserved.
+
+#include "OpenXRInputFunctionLibrary.h"
+#include "OpenXRInputSettings.h"
+
+#include "Engine/Engine.h"
+#include "Features/IModularFeatures.h"
+#include "HeadMountedDisplayFunctionLibrary.h"
+#include "PlayerMappableInputConfig.h"
+#include "UObject/ObjectPtr.h"
+#include "XRMotionControllerBase.h"
+
+UOpenXRInputFunctionLibrary::UOpenXRInputFunctionLibrary(const FObjectInitializer& ObjectInitializer)
+	: Super(ObjectInitializer)
+{
+}
+
+bool UOpenXRInputFunctionLibrary::BeginXRSession(UPlayerMappableInputConfig* InputConfig)
+{
+	if (!InputConfig)
+	{
+		// If the input config is unset, then we fall back to either the input config from the project settings.
+		// If the project settings is also unset, the input config will be set to null and we'll fall back to legacy input.
+		UOpenXRInputSettings* InputSettings = GetMutableDefault<UOpenXRInputSettings>();
+		if (InputSettings && InputSettings->MappableInputConfig.IsValid())
+		{
+			InputConfig = (UPlayerMappableInputConfig*)InputSettings->MappableInputConfig.TryLoad();
+		}
+	}
+
+	TArray<IMotionController*> MotionControllers = IModularFeatures::Get().GetModularFeatureImplementations<IMotionController>(IMotionController::GetModularFeatureName());
+	for (auto MotionController : MotionControllers)
+	{
+		if (MotionController == nullptr)
+		{
+			continue;
+		}
+
+		MotionController->SetPlayerMappableInputConfig(InputConfig);
+	}
+
+	return UHeadMountedDisplayFunctionLibrary::EnableHMD(true);
+}
+
+void UOpenXRInputFunctionLibrary::EndXRSession()
+{
+	UHeadMountedDisplayFunctionLibrary::EnableHMD(false);
+}
diff --git a/Engine/Plugins/Runtime/OpenXR/Source/OpenXRInput/Private/OpenXRInputSettings.cpp b/Engine/Plugins/Runtime/OpenXR/Source/OpenXRInput/Private/OpenXRInputSettings.cpp
--- a/Engine/Plugins/Runtime/OpenXR/Source/OpenXRInput/Private/OpenXRInputSettings.cpp
+++ b/Engine/Plugins/Runtime/OpenXR/Source/OpenXRInput/Private/OpenXRInputSettings.cpp
@@ -1,8 +1,38 @@
 // Copyright Epic Games, Inc. All Rights Reserved.
 
 #include "OpenXRInputSettings.h"
+#include "XRMotionControllerBase.h"
+#include "PlayerMappableInputConfig.h"
+
+#if WITH_EDITOR
+#include "Editor.h"
+#include "Features/IModularFeatures.h"
+#endif
 
 UOpenXRInputSettings::UOpenXRInputSettings(const FObjectInitializer& ObjectInitializer)
 	: Super(ObjectInitializer)
 {
 }
+
+#if WITH_EDITOR
+void UOpenXRInputSettings::PostEditChangeChainProperty(FPropertyChangedChainEvent& PropertyChangedEvent)
+{
+	Super::PostEditChangeChainProperty(PropertyChangedEvent);
+
+	const FName MemberPropertyName = PropertyChangedEvent.PropertyChain.GetActiveMemberNode()->GetValue()->GetFName();
+
+	if (MemberPropertyName == GET_MEMBER_NAME_CHECKED(UOpenXRInputSettings, MappableInputConfig))
+	{
+		TArray<IMotionController*> MotionControllers = IModularFeatures::Get().GetModularFeatureImplementations<IMotionController>(IMotionController::GetModularFeatureName());
+		for (auto MotionController : MotionControllers)
+		{
+			if (MotionController == nullptr)
+			{
+				continue;
+			}
+
+			MotionController->SetPlayerMappableInputConfig((UPlayerMappableInputConfig*)MappableInputConfig.TryLoad());
+		}
+	}
+}
+#endif
diff --git a/Engine/Plugins/Runtime/OpenXR/Source/OpenXRInput/Public/OpenXRInputFunctionLibrary.h b/Engine/Plugins/Runtime/OpenXR/Source/OpenXRInput/Public/OpenXRInputFunctionLibrary.h
new file mode 100644
--- /dev/null
+++ b/Engine/Plugins/Runtime/OpenXR/Source/OpenXRInput/Public/OpenXRInputFunctionLibrary.h
@@ -0,0 +1,30 @@
+// Copyright Epic Games, Inc. All Rights Reserved.
+
+#pragma once
+
+#include "CoreMinimal.h"
+#include "UObject/ObjectMacros.h"
+#include "Kismet/BlueprintFunctionLibrary.h"
+#include "OpenXRInputFunctionLibrary.generated.h"
+
+UCLASS()
+class OPENXRINPUT_API UOpenXRInputFunctionLibrary : public UBlueprintFunctionLibrary
+{
+	GENERATED_UCLASS_BODY()
+
+	/**
+	* Add a player mappable input config for Enhanced Input. This will make the tracking system aware of the actions and bindings that are used for XR motion controllers.
+	* Attaching input configs to the session can only be done once, so if multiple input configs need to be added only check AttachToSession for the last input config.
+	* 
+	* @param InputConfig		The path to the player mappable input config asset
+	* @param AttachToSession	Attach all pending configs to the running session
+	*
+	* @return			False if the input config can't be attached to the session, true otherwise
+	*/
+	UFUNCTION(BlueprintCallable, Category = "Input|XRTracking")
+	static bool BeginXRSession(class UPlayerMappableInputConfig* InputConfig);
+
+
+	UFUNCTION(BlueprintCallable, Category = "Input|XRTracking")
+	static void EndXRSession();
+};
diff --git a/Engine/Plugins/Runtime/OpenXR/Source/OpenXRInput/Public/OpenXRInputSettings.h b/Engine/Plugins/Runtime/OpenXR/Source/OpenXRInput/Public/OpenXRInputSettings.h
--- a/Engine/Plugins/Runtime/OpenXR/Source/OpenXRInput/Public/OpenXRInputSettings.h
+++ b/Engine/Plugins/Runtime/OpenXR/Source/OpenXRInput/Public/OpenXRInputSettings.h
@@ -15,11 +15,12 @@ class OPENXRINPUT_API UOpenXRInputSettings : public UObject
 	GENERATED_UCLASS_BODY()
 
 public:
-	/** Use Enhanced Input for XR instead of regular input. */
-	UPROPERTY(config, EditAnywhere, Category = "Enhanced Input")
-	bool bUseEnhancedInput = false;
-
-	/** A the mappable input config used to generate action sets for OpenXR. */
+	/** Set a mappable input config to allow OpenXR runtimes to remap the Enhanced Input actions. */
 	UPROPERTY(config, EditAnywhere, Category = "Enhanced Input", meta = (DisplayName = "Mappable Input Config for XR", AllowedClasses = "/Script/EnhancedInput.PlayerMappableInputConfig"))
 	FSoftObjectPath MappableInputConfig = nullptr;
+
+	// UObject interface
+#if WITH_EDITOR
+	virtual void PostEditChangeChainProperty(struct FPropertyChangedChainEvent& PropertyChangedEvent) override;
+#endif
 };
diff --git a/Engine/Source/Runtime/HeadMountedDisplay/Public/IMotionController.h b/Engine/Source/Runtime/HeadMountedDisplay/Public/IMotionController.h
--- a/Engine/Source/Runtime/HeadMountedDisplay/Public/IMotionController.h
+++ b/Engine/Source/Runtime/HeadMountedDisplay/Public/IMotionController.h
@@ -4,6 +4,7 @@
 
 #include "CoreMinimal.h"
 #include "UObject/ObjectMacros.h"
+#include "UObject/ObjectPtr.h"
 #include "Features/IModularFeature.h"
 
 enum class EControllerHand : uint8;
@@ -121,4 +122,14 @@ public:
 	virtual float GetCustomParameterValue(const FName MotionSource, FName ParameterName, bool& bOutValueFound) const = 0;
 
 	virtual bool GetHandJointPosition(const FName MotionSource, int jointIndex, FVector& OutPosition) const = 0;
+
+	/**
+	 * Add a player mappable input config to the motion controller. This allows the motion controller to support
+	 * Enhanced Input actions.
+	 *
+	 * @param InputConfig		The path to the player mappable input config asset
+	 *
+	 * @return			False if the input config can't be attached to the session, true otherwise
+	 */
+	virtual bool SetPlayerMappableInputConfig(TObjectPtr<class UPlayerMappableInputConfig> InputConfig = nullptr) { return true; };
 };
