Use a TStrongObjectPtr to rather than attaching to the scene root

diff --git a/Engine/Plugins/Runtime/OpenXR/Source/OpenXRInput/Private/OpenXRInput.cpp b/Engine/Plugins/Runtime/OpenXR/Source/OpenXRInput/Private/OpenXRInput.cpp
--- a/Engine/Plugins/Runtime/OpenXR/Source/OpenXRInput/Private/OpenXRInput.cpp
+++ b/Engine/Plugins/Runtime/OpenXR/Source/OpenXRInput/Private/OpenXRInput.cpp
@@ -302,7 +302,7 @@ FOpenXRInputPlugin::FOpenXRInput::FOpenXRInput(FOpenXRHMD* HMD)
 		UOpenXRInputSettings* InputSettings = GetMutableDefault<UOpenXRInputSettings>();
 		if (InputSettings && InputSettings->MappableInputConfig.IsValid())
 		{
-			MappableInputConfig = (UPlayerMappableInputConfig*)InputSettings->MappableInputConfig.TryLoad();
+			SetPlayerMappableInputConfig((UPlayerMappableInputConfig*)InputSettings->MappableInputConfig.TryLoad());
 		}
 	}
 }
@@ -393,7 +393,7 @@ bool FOpenXRInputPlugin::FOpenXRInput::BuildActions(XrSession Session)
 
 	if (MappableInputConfig)
 	{
-		BuildEnhancedActions(Profiles, MappableInputConfig);
+		BuildEnhancedActions(Profiles);
 	}
 	else
 	{
@@ -592,7 +592,7 @@ void FOpenXRInputPlugin::FOpenXRInput::BuildLegacyActions(TMap<FString, FInterac
 	}
 }
 
-void FOpenXRInputPlugin::FOpenXRInput::BuildEnhancedActions(TMap<FString, FInteractionProfile>& Profiles, TObjectPtr<UPlayerMappableInputConfig> InputConfig)
+void FOpenXRInputPlugin::FOpenXRInput::BuildEnhancedActions(TMap<FString, FInteractionProfile>& Profiles)
 {
 	if (bActionsAttached)
 	{
@@ -600,7 +600,7 @@ void FOpenXRInputPlugin::FOpenXRInput::BuildEnhancedActions(TMap<FString, FInter
 		return;
 	}
 
-	for (const TPair<TObjectPtr<UInputMappingContext>, int32> MappingContext : InputConfig->GetMappingContexts())
+	for (const TPair<TObjectPtr<UInputMappingContext>, int32> MappingContext : MappableInputConfig->GetMappingContexts())
 	{
 		FOpenXRActionSet ActionSet(Instance, MappingContext.Key->GetFName(), MappingContext.Key->ContextDescription.ToString(), MappingContext.Value, MappingContext.Key);
 		TMap<FName, int32> ActionMap;
@@ -1429,24 +1429,7 @@ bool FOpenXRInputPlugin::FOpenXRInput::SetPlayerMappableInputConfig(TObjectPtr<c
 		return false;
 	}
 
-	if (MappableInputConfig == InputConfig)
-	{
-		return true;
-	}
-
-	// Remove the previous config from the scene root
-	if (MappableInputConfig)
-	{
-		MappableInputConfig->RemoveFromRoot();
-	}
-
-	// If an input config is set, then attach it to the scene root to prevent it from being unloaded.
-	if (InputConfig)
-	{
-		InputConfig->AddToRoot();
-	}
-
-	MappableInputConfig = InputConfig;
+	MappableInputConfig = TStrongObjectPtr<class UPlayerMappableInputConfig>(InputConfig);
 	return true;
 }
 
diff --git a/Engine/Plugins/Runtime/OpenXR/Source/OpenXRInput/Private/OpenXRInput.h b/Engine/Plugins/Runtime/OpenXR/Source/OpenXRInput/Private/OpenXRInput.h
--- a/Engine/Plugins/Runtime/OpenXR/Source/OpenXRInput/Private/OpenXRInput.h
+++ b/Engine/Plugins/Runtime/OpenXR/Source/OpenXRInput/Private/OpenXRInput.h
@@ -8,6 +8,7 @@
 #include "IInputDevice.h"
 #include "IHapticDevice.h"
 #include "UObject/ObjectPtr.h"
+#include "UObject/StrongObjectPtr.h"
 
 class FOpenXRHMD;
 class UInputAction;
@@ -145,7 +146,7 @@ public:
 		TMap<EControllerHand, FOpenXRController> Controllers;
 		TMap<FName, EControllerHand> MotionSourceToControllerHandMap;
 
-		TObjectPtr<UPlayerMappableInputConfig> MappableInputConfig;
+		TStrongObjectPtr<UPlayerMappableInputConfig> MappableInputConfig;
 
 		XrAction GetActionForMotionSource(FName MotionSource) const;
 		int32 GetDeviceIDForMotionSource(FName MotionSource) const;
@@ -158,7 +159,7 @@ public:
 		bool BuildActions(XrSession Session);
 		void SyncActions(XrSession Session);
 		void BuildLegacyActions(TMap<FString, FInteractionProfile>& Profiles);
-		void BuildEnhancedActions(TMap<FString, FInteractionProfile>& Profiles, TObjectPtr<UPlayerMappableInputConfig> MappableInputConfig);
+		void BuildEnhancedActions(TMap<FString, FInteractionProfile>& Profiles);
 		void DestroyActions();
 
 		template<typename T>
