Controller action sets cannot be persistent across sessions.

Partial revert of CL 21571762

diff --git a/Engine/Plugins/Runtime/OpenXR/Source/OpenXRInput/Private/OpenXRInput.cpp b/Engine/Plugins/Runtime/OpenXR/Source/OpenXRInput/Private/OpenXRInput.cpp
--- a/Engine/Plugins/Runtime/OpenXR/Source/OpenXRInput/Private/OpenXRInput.cpp
+++ b/Engine/Plugins/Runtime/OpenXR/Source/OpenXRInput/Private/OpenXRInput.cpp
@@ -278,25 +278,9 @@ FOpenXRInputPlugin::FOpenXRInput::FOpenXRInput(FOpenXRHMD* HMD)
 		MotionSourceToControllerHandMap.Add(TEXT("EControllerHand::Right"), EControllerHand::Right);
 		MotionSourceToControllerHandMap.Add(TEXT("EControllerHand::AnyHand"), EControllerHand::AnyHand);
 
-		// Create an engine action set for pose input and haptic output
-		ControllerActionSet = MakeUnique<FOpenXRActionSet>(Instance, "controllers", "Controllers", 0);
-
-		XrPath LeftHand = GetPath(Instance, "/user/hand/left");
-		XrPath RightHand = GetPath(Instance, "/user/hand/right");
-		XrPath Head = GetPath(Instance, "/user/head");
-
-		// Controller poses
-		Controllers.Add(EControllerHand::Left, FOpenXRController(ControllerActionSet->Handle, LeftHand, "Left Controller"));
-		Controllers.Add(EControllerHand::Right, FOpenXRController(ControllerActionSet->Handle, RightHand, "Right Controller"));
-		Controllers.Add(EControllerHand::HMD, FOpenXRController(ControllerActionSet->Handle, Head, "HMD"));
-
 		// Generate a list of the sub-action paths so we can query the left/right hand individually
-		SubactionPaths.Add(LeftHand);
-		SubactionPaths.Add(RightHand);
-
-		// Make OpenXRHMD aware of the controller action spaces
-		Controllers[EControllerHand::Left].AddActionDevices(OpenXRHMD);
-		Controllers[EControllerHand::Right].AddActionDevices(OpenXRHMD);
+		SubactionPaths.Add(GetPath(Instance, "/user/hand/left"));
+		SubactionPaths.Add(GetPath(Instance, "/user/hand/right"));
 
 		// Attempt to load the default input config from the OpenXR input settings.
 		UOpenXRInputSettings* InputSettings = GetMutableDefault<UOpenXRInputSettings>();
@@ -354,11 +338,6 @@ bool FOpenXRInputPlugin::FOpenXRInput::IsOpenXRInputSupportedMotionSource(const
 FOpenXRInputPlugin::FOpenXRInput::~FOpenXRInput()
 {
 	DestroyActions();
-
-	if (ControllerActionSet)
-	{
-		xrDestroyActionSet(ControllerActionSet->Handle);
-	}
 }
 
 bool FOpenXRInputPlugin::FOpenXRInput::BuildActions(XrSession Session)
@@ -370,6 +349,22 @@ bool FOpenXRInputPlugin::FOpenXRInput::BuildActions(XrSession Session)
 
 	DestroyActions();
 
+	// Create an engine action set for pose input and haptic output
+	ControllerActionSet = MakeUnique<FOpenXRActionSet>(Instance, "controllers", "Controllers", 0);
+
+	XrPath LeftHand = GetPath(Instance, "/user/hand/left");
+	XrPath RightHand = GetPath(Instance, "/user/hand/right");
+	XrPath Head = GetPath(Instance, "/user/head");
+
+	// Controller poses
+	Controllers.Add(EControllerHand::Left, FOpenXRController(ControllerActionSet->Handle, LeftHand, "Left Controller"));
+	Controllers.Add(EControllerHand::Right, FOpenXRController(ControllerActionSet->Handle, RightHand, "Right Controller"));
+	Controllers.Add(EControllerHand::HMD, FOpenXRController(ControllerActionSet->Handle, Head, "HMD"));
+
+	// Make OpenXRHMD aware of the controller action spaces
+	Controllers[EControllerHand::Left].AddActionDevices(OpenXRHMD);
+	Controllers[EControllerHand::Right].AddActionDevices(OpenXRHMD);
+
 	// Generate a map of all supported interaction profiles to store suggested bindings
 	TMap<FString, FInteractionProfile> Profiles;
 	Profiles.Add("SimpleController", FInteractionProfile(GetPath(Instance, "/interaction_profiles/khr/simple_controller"), true));
@@ -637,16 +632,28 @@ void FOpenXRInputPlugin::FOpenXRInput::BuildEnhancedActions(TMap<FString, FInter
 
 void FOpenXRInputPlugin::FOpenXRInput::DestroyActions()
 {
+	if (OpenXRHMD)
+	{
+		OpenXRHMD->ResetActionDevices();
+	}
+
 	// Destroying an action set will also destroy all actions in the set
 	for (const FOpenXRActionSet& ActionSet : ActionSets)
 	{
 		xrDestroyActionSet(ActionSet.Handle);
 	}
 
+	if (ControllerActionSet)
+	{
+		xrDestroyActionSet(ControllerActionSet->Handle);
+		ControllerActionSet.Release();
+	}
+
 	LegacyActions.Reset();
 	EnhancedActions.Reset();
 	ActionSets.Reset();
 	PluginActionSets.Reset();
+	Controllers.Reset();
 }
 
 template<typename T>
