Implement the XR_MND_headless extension

This allows using tracking input from controllers without rendering to the HMD

diff --git a/Engine/Plugins/Runtime/OpenXR/Source/OpenXRHMD/Private/OpenXRHMD.cpp b/Engine/Plugins/Runtime/OpenXR/Source/OpenXRHMD/Private/OpenXRHMD.cpp
--- a/Engine/Plugins/Runtime/OpenXR/Source/OpenXRHMD/Private/OpenXRHMD.cpp
+++ b/Engine/Plugins/Runtime/OpenXR/Source/OpenXRHMD/Private/OpenXRHMD.cpp
@@ -1676,7 +1676,7 @@ bool FOpenXRHMD::OnStereoStartup()
 
 	XrSessionCreateInfo SessionInfo;
 	SessionInfo.type = XR_TYPE_SESSION_CREATE_INFO;
-	SessionInfo.next = RenderBridge->GetGraphicsBinding();
+	SessionInfo.next = RenderBridge.IsValid() ? RenderBridge->GetGraphicsBinding() : nullptr;
 	SessionInfo.createFlags = 0;
 	SessionInfo.systemId = System;
 	for (IOpenXRExtensionPlugin* Module : ExtensionPlugins)
@@ -1762,7 +1762,10 @@ bool FOpenXRHMD::OnStereoStartup()
 		DeviceSpace.CreateSpace(Session);
 	}
 
-	RenderBridge->SetOpenXRHMD(this);
+	if (RenderBridge.IsValid())
+	{
+		RenderBridge->SetOpenXRHMD(this);
+	}
 
 	// grab a pointer to the renderer module for displaying our mirror window
 	static const FName RendererModuleName("Renderer");
@@ -2172,6 +2175,12 @@ bool FOpenXRHMD::AllocateDepthTexture(uint32 Index, uint32 SizeX, uint32 SizeY,
 void FOpenXRHMD::OnBeginRendering_RenderThread(FRHICommandListImmediate& RHICmdList, FSceneViewFamily& ViewFamily)
 {
 	ensure(IsInRenderingThread());
+	if (!RenderBridge)
+	{
+		// Frame submission is not necessary in a headless session.
+		return;
+	}
+
 	FReadScopeLock DeviceLock(DeviceMutex);
 
 	for (IOpenXRExtensionPlugin* Module : ExtensionPlugins)
@@ -2417,7 +2426,7 @@ void FOpenXRHMD::OnBeginSimulation_GameThread()
 {
 	FReadScopeLock Lock(SessionHandleMutex);
 
-	if (!bShouldWait)
+	if (!bShouldWait || !RenderBridge)
 	{
 		return;
 	}
@@ -2656,7 +2665,7 @@ void FOpenXRHMD::OnBeginRendering_RHIThread(const FPipelinedFrameState& InFrameS
 	SCOPED_NAMED_EVENT(BeginFrame, FColor::Red);
 
 	FReadScopeLock Lock(SessionHandleMutex);
-	if (!bIsRunning)
+	if (!bIsRunning || !RenderBridge)
 	{
 		return;
 	}
@@ -2715,7 +2724,7 @@ void FOpenXRHMD::OnFinishRendering_RHIThread()
 
 	SCOPED_NAMED_EVENT(EndFrame, FColor::Red);
 
-	if (!bIsRendering)
+	if (!bIsRendering || !RenderBridge)
 	{
 		return;
 	}
diff --git a/Engine/Plugins/Runtime/OpenXR/Source/OpenXRHMD/Private/OpenXRHMD.h b/Engine/Plugins/Runtime/OpenXR/Source/OpenXRHMD/Private/OpenXRHMD.h
--- a/Engine/Plugins/Runtime/OpenXR/Source/OpenXRHMD/Private/OpenXRHMD.h
+++ b/Engine/Plugins/Runtime/OpenXR/Source/OpenXRHMD/Private/OpenXRHMD.h
@@ -302,7 +302,7 @@ public:
 	virtual bool IsActiveThisFrame_Internal(const FSceneViewExtensionContext& Context) const override;
 
 	/** IStereoRenderTargetManager */
-	virtual bool ShouldUseSeparateRenderTarget() const override { return IsStereoEnabled(); }
+	virtual bool ShouldUseSeparateRenderTarget() const override { return IsStereoEnabled() && RenderBridge.IsValid(); }
 	virtual void CalculateRenderTargetSize(const FViewport& Viewport, uint32& InOutSizeX, uint32& InOutSizeY) override;
 	virtual bool AllocateRenderTargetTexture(uint32 Index, uint32 SizeX, uint32 SizeY, uint8 Format, uint32 NumMips, ETextureCreateFlags Flags, ETextureCreateFlags TargetableTextureFlags, FTexture2DRHIRef& OutTargetableTexture, FTexture2DRHIRef& OutShaderResourceTexture, uint32 NumSamples = 1) override;
 	virtual bool NeedReAllocateDepthTexture(const TRefCountPtr<IPooledRenderTarget>& DepthTarget) override final { return bNeedReAllocatedDepth; }
diff --git a/Engine/Plugins/Runtime/OpenXR/Source/OpenXRHMD/Private/OpenXRHMDModule.cpp b/Engine/Plugins/Runtime/OpenXR/Source/OpenXRHMD/Private/OpenXRHMDModule.cpp
--- a/Engine/Plugins/Runtime/OpenXR/Source/OpenXRHMD/Private/OpenXRHMDModule.cpp
+++ b/Engine/Plugins/Runtime/OpenXR/Source/OpenXRHMD/Private/OpenXRHMDModule.cpp
@@ -42,7 +42,12 @@ FOpenXRHMDModule::~FOpenXRHMDModule()
 
 TSharedPtr< class IXRTrackingSystem, ESPMode::ThreadSafe > FOpenXRHMDModule::CreateTrackingSystem()
 {
-	if (!RenderBridge)
+	if (!InitInstanceAndSystem())
+	{
+		return nullptr;
+	}
+
+	if (!RenderBridge && !FParse::Param(FCommandLine::Get(), TEXT("xrtrackingonly")))
 	{
 		if (!InitRenderBridge())
 		{
@@ -79,6 +84,11 @@ void FOpenXRHMDModule::ShutdownModule()
 
 uint64 FOpenXRHMDModule::GetGraphicsAdapterLuid()
 {
+	if (FParse::Param(FCommandLine::Get(), TEXT("xrtrackingonly")))
+	{
+		return 0;
+	}
+
 	if (!RenderBridge)
 	{
 		if (!InitRenderBridge())
@@ -419,6 +429,14 @@ bool FOpenXRHMDModule::GetRequiredExtensions(TArray<const ANSICHAR*>& OutExtensi
 #if PLATFORM_ANDROID
 	OutExtensions.Add(XR_KHR_ANDROID_CREATE_INSTANCE_EXTENSION_NAME);
 #endif
+
+	// If the commandline -xrtrackingonly is passed, then start the application in _Other mode instead of _Scene mode
+	// This is used when we only want to get tracking information and don't need to render anything to the XR device
+	if (FParse::Param(FCommandLine::Get(), TEXT("xrtrackingonly")))
+	{
+		OutExtensions.Add(XR_MND_HEADLESS_EXTENSION_NAME);
+	}
+
 	return true;
 }
 
