Check if actions are attached before attempting to locate controllers

diff --git a/Engine/Plugins/Runtime/OpenXR/Source/OpenXRInput/Private/OpenXRInput.cpp b/Engine/Plugins/Runtime/OpenXR/Source/OpenXRInput/Private/OpenXRInput.cpp
--- a/Engine/Plugins/Runtime/OpenXR/Source/OpenXRInput/Private/OpenXRInput.cpp
+++ b/Engine/Plugins/Runtime/OpenXR/Source/OpenXRInput/Private/OpenXRInput.cpp
@@ -788,8 +788,6 @@ void FOpenXRInputPlugin::FOpenXRInput::Tick(float DeltaTime)
 	XrSession Session = OpenXRHMD->GetSession();
 	if (Session != XR_NULL_HANDLE)
 	{
-		// If we're using legacy actions we attach them to the session immediately.
-		// Otherwise we wait for a blueprint to attach enhanced input mappings.
 		if (!bActionsAttached)
 		{
 			BuildActions(Session);
@@ -1140,7 +1138,7 @@ FName FOpenXRInputPlugin::FOpenXRInput::GetMotionControllerDeviceTypeName() cons
 
 bool FOpenXRInputPlugin::FOpenXRInput::GetControllerOrientationAndPosition(const int32 ControllerIndex, const FName MotionSource, FRotator& OutOrientation, FVector& OutPosition, float WorldToMetersScale) const
 {
-	if (OpenXRHMD == nullptr)
+	if (!bActionsAttached || OpenXRHMD == nullptr)
 	{
 		return false;
 	}
@@ -1194,7 +1192,7 @@ bool FOpenXRInputPlugin::FOpenXRInput::GetControllerOrientationAndPosition(const
 
 bool FOpenXRInputPlugin::FOpenXRInput::GetControllerOrientationAndPositionForTime(const int32 ControllerIndex, const FName MotionSource, FTimespan Time, bool& OutTimeWasUsed, FRotator& OutOrientation, FVector& OutPosition, bool& OutbProvidedLinearVelocity, FVector& OutLinearVelocity, bool& OutbProvidedAngularVelocity, FVector& OutAngularVelocityRadPerSec, bool& OutbProvidedLinearAcceleration, FVector& OutLinearAcceleration, float WorldToMetersScale) const
 {
-	if (OpenXRHMD == nullptr)
+	if (!bActionsAttached || OpenXRHMD == nullptr)
 	{
 		return false;
 	}
@@ -1255,7 +1253,7 @@ bool FOpenXRInputPlugin::FOpenXRInput::GetControllerOrientationAndPositionForTim
 
 ETrackingStatus FOpenXRInputPlugin::FOpenXRInput::GetControllerTrackingStatus(const int32 ControllerIndex, const FName MotionSource) const
 {
-	if (OpenXRHMD == nullptr)
+	if (!bActionsAttached || OpenXRHMD == nullptr)
 	{
 		return ETrackingStatus::NotTracked;
 	}
