Removed including ScenePrivate.h from several modules.

diff --git a/Engine/Plugins/Experimental/GPULightmass/Source/GPULightmass/Private/GPULightmass.h b/Engine/Plugins/Experimental/GPULightmass/Source/GPULightmass/Private/GPULightmass.h
--- a/Engine/Plugins/Experimental/GPULightmass/Source/GPULightmass/Private/GPULightmass.h
+++ b/Engine/Plugins/Experimental/GPULightmass/Source/GPULightmass/Private/GPULightmass.h
@@ -2,7 +2,6 @@
 
 #pragma once
 
-#include "ScenePrivate.h"
 #include "VolumetricLightmap.h"
 #include "GPULightmassModule.h"
 #include "GPULightmassSettings.h"
diff --git a/Engine/Plugins/Experimental/GPULightmass/Source/GPULightmass/Private/Scene/Scene.cpp b/Engine/Plugins/Experimental/GPULightmass/Source/GPULightmass/Private/Scene/Scene.cpp
--- a/Engine/Plugins/Experimental/GPULightmass/Source/GPULightmass/Private/Scene/Scene.cpp
+++ b/Engine/Plugins/Experimental/GPULightmass/Source/GPULightmass/Private/Scene/Scene.cpp
@@ -26,6 +26,7 @@
 #include "ShaderCompiler.h"
 #include "Misc/FileHelper.h"
 #include "Components/ReflectionCaptureComponent.h"
+#include "ScenePrivate.h"
 
 #define LOCTEXT_NAMESPACE "StaticLightingSystem"
 
diff --git a/Engine/Plugins/Experimental/GPULightmass/Source/GPULightmass/Private/VolumetricLightmap.cpp b/Engine/Plugins/Experimental/GPULightmass/Source/GPULightmass/Private/VolumetricLightmap.cpp
--- a/Engine/Plugins/Experimental/GPULightmass/Source/GPULightmass/Private/VolumetricLightmap.cpp
+++ b/Engine/Plugins/Experimental/GPULightmass/Source/GPULightmass/Private/VolumetricLightmap.cpp
@@ -2,7 +2,6 @@
 
 #include "VolumetricLightmap.h"
 #include "VolumetricLightmapVoxelization.h"
-#include "ScenePrivate.h"
 #include "Containers/DynamicRHIResourceArray.h"
 #include "LightmapRayTracing.h"
 #include "PathTracingLightParameters.inl"
diff --git a/Engine/Plugins/Runtime/OpenXR/Source/OpenXRHMD/Private/OpenXRHMD.cpp b/Engine/Plugins/Runtime/OpenXR/Source/OpenXRHMD/Private/OpenXRHMD.cpp
--- a/Engine/Plugins/Runtime/OpenXR/Source/OpenXRHMD/Private/OpenXRHMD.cpp
+++ b/Engine/Plugins/Runtime/OpenXR/Source/OpenXRHMD/Private/OpenXRHMD.cpp
@@ -29,10 +29,11 @@
 #include "ARSystem.h"
 #include "IHandTracker.h"
 #include "PixelShaderUtils.h"
-#include "ScenePrivate.h"
 #include "GeneralProjectSettings.h"
 #include "Epic_openxr.h"
 #include "HDRHelper.h"
+#include "Shader.h"
+#include "ScreenRendering.h"
 
 #if WITH_EDITOR
 #include "Editor/UnrealEd/Classes/Editor/EditorEngine.h"
diff --git a/Engine/Plugins/Runtime/Steam/SteamVR/Source/SteamVR/Private/SteamVRHMD.cpp b/Engine/Plugins/Runtime/Steam/SteamVR/Source/SteamVR/Private/SteamVRHMD.cpp
--- a/Engine/Plugins/Runtime/Steam/SteamVR/Source/SteamVR/Private/SteamVRHMD.cpp
+++ b/Engine/Plugins/Runtime/Steam/SteamVR/Source/SteamVR/Private/SteamVRHMD.cpp
@@ -8,7 +8,6 @@
 #include "Misc/EngineVersion.h"
 #include "HardwareInfo.h"
 #include "RendererPrivate.h"
-#include "ScenePrivate.h"
 #include "Slate/SceneViewport.h"
 #include "PostProcess/PostProcessHMD.h"
 #include "SteamVRFunctionLibrary.h"
diff --git a/Engine/Plugins/Runtime/Steam/SteamVR/Source/SteamVR/Private/SteamVRRender.cpp b/Engine/Plugins/Runtime/Steam/SteamVR/Source/SteamVR/Private/SteamVRRender.cpp
--- a/Engine/Plugins/Runtime/Steam/SteamVR/Source/SteamVR/Private/SteamVRRender.cpp
+++ b/Engine/Plugins/Runtime/Steam/SteamVR/Source/SteamVR/Private/SteamVRRender.cpp
@@ -8,7 +8,6 @@
 #include "SteamVRHMD.h"
 
 #include "RendererPrivate.h"
-#include "ScenePrivate.h"
 #include "PostProcess/PostProcessHMD.h"
 #include "PipelineStateCache.h"
 #include "ClearQuad.h"
diff --git a/Engine/Plugins/Runtime/Steam/SteamVR/Source/SteamVR/Private/SteamVRSpectatorScreen.cpp b/Engine/Plugins/Runtime/Steam/SteamVR/Source/SteamVR/Private/SteamVRSpectatorScreen.cpp
--- a/Engine/Plugins/Runtime/Steam/SteamVR/Source/SteamVR/Private/SteamVRSpectatorScreen.cpp
+++ b/Engine/Plugins/Runtime/Steam/SteamVR/Source/SteamVR/Private/SteamVRSpectatorScreen.cpp
@@ -8,7 +8,6 @@
 #include "SteamVRHMD.h"
 
 #include "RendererPrivate.h"
-#include "ScenePrivate.h"
 #include "PostProcess/PostProcessHMD.h"
 #include "PipelineStateCache.h"
 #include "ClearQuad.h"
