Hololens Remoting OpenXR fix
-Refactoring OpenXRInput so that it builds and cleans up actions and action sets at begin and end session rather than during it's tick.
-Move SyncActions from Tick() into SendControllerEvents(), mainly to make the order in which the code runs more obvious.
-Did a pass on DeviceMutex fixing a couple missing blocks and reducing the scope of others.

diff --git a/Engine/Plugins/Runtime/OpenXR/Source/OpenXRHMD/Private/OpenXRHMD.cpp b/Engine/Plugins/Runtime/OpenXR/Source/OpenXRHMD/Private/OpenXRHMD.cpp
--- a/Engine/Plugins/Runtime/OpenXR/Source/OpenXRHMD/Private/OpenXRHMD.cpp
+++ b/Engine/Plugins/Runtime/OpenXR/Source/OpenXRHMD/Private/OpenXRHMD.cpp
@@ -418,9 +418,12 @@ bool FOpenXRHMD::GetTrackingOriginTransform(TEnumAsByte<EHMDTrackingOrigin::Type
 	switch (Origin)
 	{
 	case EHMDTrackingOrigin::Eye:
-		if (DeviceSpaces.Num())
 		{
-			Space = DeviceSpaces[HMDDeviceId].Space;
+			FReadScopeLock DeviceLock(DeviceMutex);
+			if (DeviceSpaces.Num())
+			{
+				Space = DeviceSpaces[HMDDeviceId].Space;
+			}
 		}
 		break;
 	case EHMDTrackingOrigin::Floor:
@@ -689,7 +692,12 @@ void FOpenXRHMD::Recenter(EOrientPositionSelector::Type Selector, float Yaw)
 	const XrTime TargetTime = PipelineState.FrameState.predictedDisplayTime;
 	check(PipelineState.bXrFrameStateUpdated);
 
-	const FDeviceSpace& DeviceSpace = DeviceSpaces[HMDDeviceId];
+	XrSpace DeviceSpace = XR_NULL_HANDLE;
+	{
+		FReadScopeLock DeviceLock(DeviceMutex);
+		const FDeviceSpace& DeviceSpaceStruct = DeviceSpaces[HMDDeviceId];
+		DeviceSpace = DeviceSpaceStruct.Space;
+	}
 	XrSpaceLocation DeviceLocation = { XR_TYPE_SPACE_LOCATION, nullptr };
 
 	XrSpace BaseSpace = TrackingSpaceType == XR_REFERENCE_SPACE_TYPE_STAGE ? StageSpace : LocalSpace;
@@ -697,7 +705,7 @@ void FOpenXRHMD::Recenter(EOrientPositionSelector::Type Selector, float Yaw)
 	{
 		BaseSpace = CustomSpace;
 	}
-	XR_ENSURE(xrLocateSpace(DeviceSpace.Space, BaseSpace, TargetTime, &DeviceLocation));
+	XR_ENSURE(xrLocateSpace(DeviceSpace, BaseSpace, TargetTime, &DeviceLocation));
 
 	const FQuat CurrentOrientation = ToFQuat(DeviceLocation.pose.orientation);
 	const FVector CurrentPosition = ToFVector(DeviceLocation.pose.position, GetWorldToMetersScale());
@@ -1201,6 +1209,7 @@ FOpenXRHMD::FOpenXRHMD(const FAutoRegister& AutoRegister, XrInstance InInstance,
 	, bIsStandaloneStereoOnlyDevice(false)
 	, CurrentSessionState(XR_SESSION_STATE_UNKNOWN)
 	, EnabledExtensions(std::move(InEnabledExtensions))
+	, InputModule(nullptr)
 	, ExtensionPlugins(std::move(InExtensionPlugins))
 	, Instance(InInstance)
 	, System(InSystem)
@@ -1670,7 +1679,6 @@ bool FOpenXRHMD::BuildOcclusionMesh(XrVisibilityMaskTypeKHR Type, int View, FHMD
 bool FOpenXRHMD::OnStereoStartup()
 {
 	FWriteScopeLock Lock(SessionHandleMutex);
-	FWriteScopeLock DeviceLock(DeviceMutex);
 
 	check(Session == XR_NULL_HANDLE);
 	bIsExitingSessionByxrRequestExitSession = false;  // clear in case we requested exit for a previous session, but it ended in some other way before that happened.
@@ -1701,7 +1709,7 @@ bool FOpenXRHMD::OnStereoStartup()
 	TArray<XrReferenceSpaceType> ReferenceSpaces;
 	ReferenceSpaces.SetNum(ReferenceSpacesCount);
 	// Initialize spaces array with valid enum values (avoid triggering validation error).
-	for (auto & SpaceIter : ReferenceSpaces)
+	for (auto& SpaceIter : ReferenceSpaces)
 		SpaceIter = XR_REFERENCE_SPACE_TYPE_VIEW;
 	XR_ENSURE(xrEnumerateReferenceSpaces(Session, (uint32_t)ReferenceSpaces.Num(), &ReferenceSpacesCount, ReferenceSpaces.GetData()));
 	ensure(ReferenceSpacesCount == ReferenceSpaces.Num());
@@ -1714,7 +1722,10 @@ bool FOpenXRHMD::OnStereoStartup()
 	SpaceInfo.referenceSpaceType = XR_REFERENCE_SPACE_TYPE_VIEW;
 	SpaceInfo.poseInReferenceSpace = ToXrPose(FTransform::Identity);
 	XR_ENSURE(xrCreateReferenceSpace(Session, &SpaceInfo, &HmdSpace));
-	DeviceSpaces[HMDDeviceId].Space = HmdSpace;
+	{
+		FWriteScopeLock DeviceLock(DeviceMutex);
+		DeviceSpaces[HMDDeviceId].Space = HmdSpace;
+	}
 
 	ensure(ReferenceSpaces.Contains(XR_REFERENCE_SPACE_TYPE_LOCAL));
 	SpaceInfo.referenceSpaceType = XR_REFERENCE_SPACE_TYPE_LOCAL;
@@ -1758,9 +1769,12 @@ bool FOpenXRHMD::OnStereoStartup()
 	PipelinedFrameStateGame.TrackingSpace->CreateSpace(Session);
 
 	// Create action spaces for all devices
-	for (FDeviceSpace& DeviceSpace : DeviceSpaces)
 	{
-		DeviceSpace.CreateSpace(Session);
+		FWriteScopeLock DeviceLock(DeviceMutex);
+		for (FDeviceSpace& DeviceSpace : DeviceSpaces)
+		{
+			DeviceSpace.CreateSpace(Session);
+		}
 	}
 
 	if (RenderBridge.IsValid())
@@ -1843,13 +1857,13 @@ void FOpenXRHMD::DestroySession()
 
 	if (Session != XR_NULL_HANDLE)
 	{
-		FWriteScopeLock DeviceLock(DeviceMutex);
-
 		for (IOpenXRExtensionPlugin* Module : ExtensionPlugins)
 		{
 			Module->OnDestroySession(Session);
 		}
 
+		InputModule->OnDestroySession();
+
 		// We need to reset all swapchain references to ensure there are no attempts
 		// to destroy swapchain handles after the session is already destroyed.
 		ForEachLayer([&](uint32 /* unused */, FOpenXRLayer& Layer)
@@ -1885,15 +1899,19 @@ void FOpenXRHMD::DestroySession()
 
 		// Destroy device and reference spaces, they will be recreated
 		// when the session is created again.
-		for (FDeviceSpace& Device : DeviceSpaces)
 		{
-			Device.DestroySpace();
+			FReadScopeLock DeviceLock(DeviceMutex);
+			for (FDeviceSpace& Device : DeviceSpaces)
+			{
+				Device.DestroySpace();
+			}
 		}
 
 		// Close the session now we're allowed to.
 		XR_ENSURE(xrDestroySession(Session));
 		Session = XR_NULL_HANDLE;
-
+		CurrentSessionState = XR_SESSION_STATE_UNKNOWN;
+		UE_LOG(LogHMD, Verbose, TEXT("Session state switched to XR_SESSION_STATE_UNKNOWN by DestroySession()"), OpenXRSessionStateToString(CurrentSessionState));
 		bStereoEnabled = false;
 		bIsReady = false;
 		bIsRunning = false;
@@ -1913,7 +1931,7 @@ int32 FOpenXRHMD::AddActionDevice(XrAction Action, XrPath Path)
 
 	int32 DeviceId = DeviceSpaces.Emplace(Action, Path);
 
-	FReadScopeLock Lock(SessionHandleMutex);
+	//FReadScopeLock Lock(SessionHandleMutex); // This is called from StartSession(), which already has this lock.
 	if (Session)
 	{
 		DeviceSpaces[DeviceId].CreateSpace(Session);
@@ -1987,11 +2005,15 @@ bool FOpenXRHMD::StartSession()
 		return false;
 	}
 
+	check(InputModule);
+	InputModule->OnBeginSession();
+
 	XrSessionBeginInfo Begin = { XR_TYPE_SESSION_BEGIN_INFO, nullptr, SelectedViewConfigurationType };
 	for (IOpenXRExtensionPlugin* Module : ExtensionPlugins)
 	{
 		Begin.next = Module->OnBeginSession(Session, Begin.next);
 	}
+
 	bIsRunning = XR_ENSURE(xrBeginSession(Session, &Begin));
 	return bIsRunning;
 }
@@ -2182,8 +2204,6 @@ void FOpenXRHMD::OnBeginRendering_RenderThread(FRHICommandListImmediate& RHICmdL
 		return;
 	}
 
-	FReadScopeLock DeviceLock(DeviceMutex);
-
 	for (IOpenXRExtensionPlugin* Module : ExtensionPlugins)
 	{
 		Module->OnBeginRendering_RenderThread(Session);
@@ -2259,55 +2279,59 @@ void FOpenXRHMD::OnBeginRendering_RenderThread(FRHICommandListImmediate& RHICmdL
 	CopySortedLayers(StereoLayers);
 	PipelinedLayerStateRendering.QuadLayers.Reset(StereoLayers.Num());
 	PipelinedLayerStateRendering.QuadSwapchains.Reset(StereoLayers.Num());
-	for (const FOpenXRLayer& Layer : StereoLayers)
 	{
-		const bool bNoAlpha = Layer.Desc.Flags & IStereoLayers::LAYER_FLAG_TEX_NO_ALPHA_CHANNEL;
-		const bool bIsStereo = Layer.Desc.LeftTexture.IsValid();
-		FTransform PositionTransform = Layer.Desc.PositionType == ELayerType::WorldLocked ?
-			InvTrackingToWorld : FTransform::Identity;
-
-		XrCompositionLayerQuad Quad = { XR_TYPE_COMPOSITION_LAYER_QUAD, nullptr };
-		Quad.layerFlags = bNoAlpha ? 0 : XR_COMPOSITION_LAYER_UNPREMULTIPLIED_ALPHA_BIT |
-			XR_COMPOSITION_LAYER_BLEND_TEXTURE_SOURCE_ALPHA_BIT;
-		Quad.space = Layer.Desc.PositionType == ELayerType::FaceLocked ?
-			DeviceSpaces[HMDDeviceId].Space : PipelinedFrameStateRendering.TrackingSpace->Handle;
-		Quad.subImage.imageArrayIndex = 0;
-		Quad.pose = ToXrPose(Layer.Desc.Transform * PositionTransform, WorldToMeters);
+		FReadScopeLock DeviceLock(DeviceMutex);
 
-		// We need to copy each layer into an OpenXR swapchain so they can be displayed by the compositor
-		if (Layer.RightEye.Swapchain.IsValid() && Layer.Desc.Texture.IsValid())
+		for (const FOpenXRLayer& Layer : StereoLayers)
 		{
-			if (Layer.RightEye.bUpdateTexture && bIsRunning)
+			const bool bNoAlpha = Layer.Desc.Flags & IStereoLayers::LAYER_FLAG_TEX_NO_ALPHA_CHANNEL;
+			const bool bIsStereo = Layer.Desc.LeftTexture.IsValid();
+			FTransform PositionTransform = Layer.Desc.PositionType == ELayerType::WorldLocked ?
+				InvTrackingToWorld : FTransform::Identity;
+
+			XrCompositionLayerQuad Quad = { XR_TYPE_COMPOSITION_LAYER_QUAD, nullptr };
+			Quad.layerFlags = bNoAlpha ? 0 : XR_COMPOSITION_LAYER_UNPREMULTIPLIED_ALPHA_BIT |
+				XR_COMPOSITION_LAYER_BLEND_TEXTURE_SOURCE_ALPHA_BIT;
+			Quad.space = Layer.Desc.PositionType == ELayerType::FaceLocked ?
+				DeviceSpaces[HMDDeviceId].Space : PipelinedFrameStateRendering.TrackingSpace->Handle;
+			Quad.subImage.imageArrayIndex = 0;
+			Quad.pose = ToXrPose(Layer.Desc.Transform * PositionTransform, WorldToMeters);
+
+			// We need to copy each layer into an OpenXR swapchain so they can be displayed by the compositor
+			if (Layer.RightEye.Swapchain.IsValid() && Layer.Desc.Texture.IsValid())
 			{
-				FRHITexture2D* SrcTexture = Layer.Desc.Texture->GetTexture2D();
-				FIntRect DstRect(FIntPoint(0, 0), Layer.RightEye.SwapchainSize.IntPoint());
-				CopyTexture_RenderThread(RHICmdList, SrcTexture, FIntRect(), Layer.RightEye.Swapchain, DstRect, false, bNoAlpha);
-			}
+				if (Layer.RightEye.bUpdateTexture && bIsRunning)
+				{
+					FRHITexture2D* SrcTexture = Layer.Desc.Texture->GetTexture2D();
+					FIntRect DstRect(FIntPoint(0, 0), Layer.RightEye.SwapchainSize.IntPoint());
+					CopyTexture_RenderThread(RHICmdList, SrcTexture, FIntRect(), Layer.RightEye.Swapchain, DstRect, false, bNoAlpha);
+				}
 
-			Quad.eyeVisibility = bIsStereo ? XR_EYE_VISIBILITY_RIGHT : XR_EYE_VISIBILITY_BOTH;
+				Quad.eyeVisibility = bIsStereo ? XR_EYE_VISIBILITY_RIGHT : XR_EYE_VISIBILITY_BOTH;
 
-			Quad.subImage.imageRect = ToXrRect(Layer.GetRightViewportSize());
-			Quad.subImage.swapchain = static_cast<FOpenXRSwapchain*>(Layer.RightEye.Swapchain.Get())->GetHandle();
-			Quad.size = ToXrExtent2D(Layer.GetRightQuadSize(), WorldToMeters);
-			PipelinedLayerStateRendering.QuadLayers.Add(Quad);
-			PipelinedLayerStateRendering.QuadSwapchains.Add(Layer.RightEye.Swapchain);
-		}
+				Quad.subImage.imageRect = ToXrRect(Layer.GetRightViewportSize());
+				Quad.subImage.swapchain = static_cast<FOpenXRSwapchain*>(Layer.RightEye.Swapchain.Get())->GetHandle();
+				Quad.size = ToXrExtent2D(Layer.GetRightQuadSize(), WorldToMeters);
+				PipelinedLayerStateRendering.QuadLayers.Add(Quad);
+				PipelinedLayerStateRendering.QuadSwapchains.Add(Layer.RightEye.Swapchain);
+			}
 
-		if (Layer.LeftEye.Swapchain.IsValid() && Layer.Desc.LeftTexture.IsValid())
-		{
-			if (Layer.LeftEye.bUpdateTexture && bIsRunning)
+			if (Layer.LeftEye.Swapchain.IsValid() && Layer.Desc.LeftTexture.IsValid())
 			{
-				FRHITexture2D* SrcTexture = Layer.Desc.LeftTexture->GetTexture2D();
-				FIntRect DstRect(FIntPoint(0, 0), Layer.LeftEye.SwapchainSize.IntPoint());
-				CopyTexture_RenderThread(RHICmdList, SrcTexture, FIntRect(), Layer.LeftEye.Swapchain, DstRect, false, bNoAlpha);
-			}
+				if (Layer.LeftEye.bUpdateTexture && bIsRunning)
+				{
+					FRHITexture2D* SrcTexture = Layer.Desc.LeftTexture->GetTexture2D();
+					FIntRect DstRect(FIntPoint(0, 0), Layer.LeftEye.SwapchainSize.IntPoint());
+					CopyTexture_RenderThread(RHICmdList, SrcTexture, FIntRect(), Layer.LeftEye.Swapchain, DstRect, false, bNoAlpha);
+				}
 
-			Quad.eyeVisibility = XR_EYE_VISIBILITY_LEFT;
-			Quad.subImage.imageRect = ToXrRect(Layer.GetLeftViewportSize());
-			Quad.subImage.swapchain = static_cast<FOpenXRSwapchain*>(Layer.LeftEye.Swapchain.Get())->GetHandle();
-			Quad.size = ToXrExtent2D(Layer.GetLeftQuadSize(), WorldToMeters);
-			PipelinedLayerStateRendering.QuadLayers.Add(Quad);
-			PipelinedLayerStateRendering.QuadSwapchains.Add(Layer.LeftEye.Swapchain);
+				Quad.eyeVisibility = XR_EYE_VISIBILITY_LEFT;
+				Quad.subImage.imageRect = ToXrRect(Layer.GetLeftViewportSize());
+				Quad.subImage.swapchain = static_cast<FOpenXRSwapchain*>(Layer.LeftEye.Swapchain.Get())->GetHandle();
+				Quad.size = ToXrExtent2D(Layer.GetLeftQuadSize(), WorldToMeters);
+				PipelinedLayerStateRendering.QuadLayers.Add(Quad);
+				PipelinedLayerStateRendering.QuadSwapchains.Add(Layer.LeftEye.Swapchain);
+			}
 		}
 	}
 
diff --git a/Engine/Plugins/Runtime/OpenXR/Source/OpenXRHMD/Private/OpenXRHMD.h b/Engine/Plugins/Runtime/OpenXR/Source/OpenXRHMD/Private/OpenXRHMD.h
--- a/Engine/Plugins/Runtime/OpenXR/Source/OpenXRHMD/Private/OpenXRHMD.h
+++ b/Engine/Plugins/Runtime/OpenXR/Source/OpenXRHMD/Private/OpenXRHMD.h
@@ -23,6 +23,7 @@ class FSceneView;
 class FSceneViewFamily;
 class UCanvas;
 class FOpenXRRenderBridge;
+class IOpenXRInputModule;
 
 /**
  * Simple Head Mounted Display
@@ -328,6 +329,11 @@ public:
 	/** Constructor */
 	FOpenXRHMD(const FAutoRegister&, XrInstance InInstance, XrSystemId InSystem, TRefCountPtr<FOpenXRRenderBridge>& InRenderBridge, TArray<const char*> InEnabledExtensions, TArray<class IOpenXRExtensionPlugin*> InExtensionPlugins, IARSystemSupport* ARSystemSupport);
 
+	void SetInputModule(IOpenXRInputModule* InInputModule)
+	{
+		InputModule = InInputModule;
+	}
+
 	/** Destructor */
 	virtual ~FOpenXRHMD();
 
@@ -377,6 +383,7 @@ private:
 	FRWLock					SessionHandleMutex;
 
 	TArray<const char*>		EnabledExtensions;
+	IOpenXRInputModule*		InputModule;
 	TArray<class IOpenXRExtensionPlugin*> ExtensionPlugins;
 	XrInstance				Instance;
 	XrSystemId				System;
@@ -397,7 +404,7 @@ private:
 	FPipelinedLayerState	PipelinedLayerStateRendering;
 	FPipelinedLayerState	PipelinedLayerStateRHI;
 
-	FRWLock					DeviceMutex;
+	mutable FRWLock			DeviceMutex;
 	TArray<FDeviceSpace>	DeviceSpaces;
 
 	TRefCountPtr<FOpenXRRenderBridge> RenderBridge;
diff --git a/Engine/Plugins/Runtime/OpenXR/Source/OpenXRHMD/Public/IOpenXRExtensionPlugin.h b/Engine/Plugins/Runtime/OpenXR/Source/OpenXRHMD/Public/IOpenXRExtensionPlugin.h
--- a/Engine/Plugins/Runtime/OpenXR/Source/OpenXRHMD/Public/IOpenXRExtensionPlugin.h
+++ b/Engine/Plugins/Runtime/OpenXR/Source/OpenXRHMD/Public/IOpenXRExtensionPlugin.h
@@ -100,6 +100,15 @@ public:
 	}
 };
 
+// Note: We may refactor to put OpenXRInput into the OpenXRHMD module so we can get rid of this interface.
+class IOpenXRInputModule
+{
+public:
+	virtual ~IOpenXRInputModule() {}
+
+	virtual void OnBeginSession() = 0;
+	virtual void OnDestroySession() = 0;
+};
 
 class IOpenXRExtensionPlugin : public IModularFeature
 {
diff --git a/Engine/Plugins/Runtime/OpenXR/Source/OpenXRInput/Private/OpenXRInput.cpp b/Engine/Plugins/Runtime/OpenXR/Source/OpenXRInput/Private/OpenXRInput.cpp
--- a/Engine/Plugins/Runtime/OpenXR/Source/OpenXRInput/Private/OpenXRInput.cpp
+++ b/Engine/Plugins/Runtime/OpenXR/Source/OpenXRInput/Private/OpenXRInput.cpp
@@ -279,6 +279,8 @@ FOpenXRInputPlugin::FOpenXRInput::FOpenXRInput(FOpenXRHMD* HMD)
 		// Generate a list of the sub-action paths so we can query the left/right hand individually
 		SubactionPaths.Add(GetPath(Instance, "/user/hand/left"));
 		SubactionPaths.Add(GetPath(Instance, "/user/hand/right"));
+
+		OpenXRHMD->SetInputModule(this);
 	}
 }
 
@@ -715,13 +717,9 @@ bool FOpenXRInputPlugin::FOpenXRInput::SuggestBindingForKey(TMap<FString, FInter
 	return SuggestBindingForKey(Profiles, Action, InFKey, Modifiers, Triggers);
 }
 
-void FOpenXRInputPlugin::FOpenXRInput::Tick(float DeltaTime)
+void FOpenXRInputPlugin::FOpenXRInput::OnBeginSession()
 {
-	if (OpenXRHMD == nullptr)
-	{
-		// In the editor, when we are not actually running OpenXR, but the IInputDevice exists so it can enumerate its motion sources.
-		return;
-	}
+	check(OpenXRHMD); // If there is no hmd this function should not be called.
 
 	XrSession Session = OpenXRHMD->GetSession();
 	if (Session != XR_NULL_HANDLE)
@@ -730,10 +728,12 @@ void FOpenXRInputPlugin::FOpenXRInput::Tick(float DeltaTime)
 		{
 			BuildActions(Session);
 		}
-
-		SyncActions(Session);
 	}
-	else if (bActionsAttached)
+}
+
+void FOpenXRInputPlugin::FOpenXRInput::OnDestroySession()
+{
+	if (bActionsAttached)
 	{
 		// If the session shut down, clean up.
 		bActionsAttached = false;
@@ -831,6 +831,11 @@ void FOpenXRInputPlugin::FOpenXRInput::SendControllerEvents()
 
 	XrSession Session = OpenXRHMD->GetSession();
 
+	if (Session != XR_NULL_HANDLE)
+	{
+		SyncActions(Session);
+	}
+
 	IPlatformInputDeviceMapper& DeviceMapper = IPlatformInputDeviceMapper::Get();
 	for (FOpenXRAction& Action : LegacyActions)
 	{
diff --git a/Engine/Plugins/Runtime/OpenXR/Source/OpenXRInput/Private/OpenXRInput.h b/Engine/Plugins/Runtime/OpenXR/Source/OpenXRInput/Private/OpenXRInput.h
--- a/Engine/Plugins/Runtime/OpenXR/Source/OpenXRInput/Private/OpenXRInput.h
+++ b/Engine/Plugins/Runtime/OpenXR/Source/OpenXRInput/Private/OpenXRInput.h
@@ -5,6 +5,7 @@
 #include "GenericPlatform/IInputInterface.h"
 #include "XRMotionControllerBase.h"
 #include "IOpenXRInputPlugin.h"
+#include "IOpenXRExtensionPlugin.h"
 #include "IInputDevice.h"
 #include "IHapticDevice.h"
 #include "UObject/ObjectPtr.h"
@@ -100,14 +101,18 @@ public:
 		FInteractionProfile(XrPath InProfile, bool InHasHaptics);
 	};
 
-	class FOpenXRInput : public IInputDevice, public FXRMotionControllerBase, public IHapticDevice, public TSharedFromThis<FOpenXRInput>
+	class FOpenXRInput : public IOpenXRInputModule, public IInputDevice, public FXRMotionControllerBase, public IHapticDevice, public TSharedFromThis<FOpenXRInput>
 	{
 	public:
 		FOpenXRInput(FOpenXRHMD* HMD);
 		virtual ~FOpenXRInput();
 
+		// IOpenXRAdditionalModule overrides
+		virtual void OnBeginSession() override;
+		virtual void OnDestroySession() override;
+
 		// IInputDevice overrides
-		virtual void Tick(float DeltaTime) override;
+		virtual void Tick(float DeltaTime) override {};
 		virtual void SendControllerEvents() override;
 		virtual void SetMessageHandler(const TSharedRef< FGenericApplicationMessageHandler >& InMessageHandler) override;
 		virtual bool Exec(UWorld* InWorld, const TCHAR* Cmd, FOutputDevice& Ar) override;
