Don't attempt to destroy actions when shutting down.

Caused a deadlock and could result in redundant destruction.
These will be destroyed when the instance is destroyed on module shutdown.

diff --git a/Engine/Plugins/Runtime/OpenXR/Source/OpenXRInput/Private/OpenXRInput.cpp b/Engine/Plugins/Runtime/OpenXR/Source/OpenXRInput/Private/OpenXRInput.cpp
--- a/Engine/Plugins/Runtime/OpenXR/Source/OpenXRInput/Private/OpenXRInput.cpp
+++ b/Engine/Plugins/Runtime/OpenXR/Source/OpenXRInput/Private/OpenXRInput.cpp
@@ -328,11 +328,6 @@ bool FOpenXRInputPlugin::FOpenXRInput::IsOpenXRInputSupportedMotionSource(const
 		|| MotionSourceToControllerHandMap.Contains(MotionSource);
 }
 
-FOpenXRInputPlugin::FOpenXRInput::~FOpenXRInput()
-{
-	DestroyActions();
-}
-
 bool FOpenXRInputPlugin::FOpenXRInput::BuildActions(XrSession Session)
 {
 	if (bActionsAttached)
diff --git a/Engine/Plugins/Runtime/OpenXR/Source/OpenXRInput/Private/OpenXRInput.h b/Engine/Plugins/Runtime/OpenXR/Source/OpenXRInput/Private/OpenXRInput.h
--- a/Engine/Plugins/Runtime/OpenXR/Source/OpenXRInput/Private/OpenXRInput.h
+++ b/Engine/Plugins/Runtime/OpenXR/Source/OpenXRInput/Private/OpenXRInput.h
@@ -105,7 +105,7 @@ public:
 	{
 	public:
 		FOpenXRInput(FOpenXRHMD* HMD);
-		virtual ~FOpenXRInput();
+		virtual ~FOpenXRInput() {};
 
 		// IOpenXRAdditionalModule overrides
 		virtual void OnBeginSession() override;
