OpenXR null pipelined tracking space when we have exited VR.
-We are nulling out the TrackingSpace in the piplined state when the session is not in a state where it can locate.  Need to not try to use it when its null.

diff --git a/Engine/Plugins/Runtime/OpenXR/Source/OpenXRHMD/Private/OpenXRHMD.cpp b/Engine/Plugins/Runtime/OpenXR/Source/OpenXRHMD/Private/OpenXRHMD.cpp
--- a/Engine/Plugins/Runtime/OpenXR/Source/OpenXRHMD/Private/OpenXRHMD.cpp
+++ b/Engine/Plugins/Runtime/OpenXR/Source/OpenXRHMD/Private/OpenXRHMD.cpp
@@ -448,6 +448,12 @@ bool FOpenXRHMD::GetTrackingOriginTransform(TEnumAsByte<EHMDTrackingOrigin::Type
 
 	const FPipelinedFrameState& PipelinedState = GetPipelinedFrameStateForThread();
 
+	if (!PipelinedState.TrackingSpace.IsValid())
+	{
+		// Session is in a state where we can't locate.
+		return false;
+	}
+
 	XrSpaceLocation NewLocation = { XR_TYPE_SPACE_LOCATION };
 	const XrResult Result = xrLocateSpace(Space, PipelinedState.TrackingSpace->Handle, PipelinedState.FrameState.predictedDisplayTime, &NewLocation);
 	if (Result != XR_SUCCESS)
@@ -1398,7 +1404,8 @@ void FOpenXRHMD::UpdateDeviceLocations(bool bUpdateOpenXRExtensionPlugins)
 	FPipelinedFrameState& PipelineState = GetPipelinedFrameStateForThread();
 
 	// Only update the device locations if the frame state has been predicted, which is dependent on WaitFrame success
-	if (PipelineState.bXrFrameStateUpdated)
+	// Also need a valid TrackingSpace
+	if (PipelineState.bXrFrameStateUpdated && PipelineState.TrackingSpace.IsValid())
 	{
 		FReadScopeLock Lock(DeviceMutex);
 		PipelineState.DeviceLocations.SetNumZeroed(DeviceSpaces.Num());
