Implement support for XR_KHR_vulkan_swapchain_format_list.

This allows for improved efficiency in the compositor.

diff --git a/Engine/Plugins/Runtime/OpenXR/Source/OpenXRHMD/Private/OpenXRHMDModule.cpp b/Engine/Plugins/Runtime/OpenXR/Source/OpenXRHMD/Private/OpenXRHMDModule.cpp
--- a/Engine/Plugins/Runtime/OpenXR/Source/OpenXRHMD/Private/OpenXRHMDModule.cpp
+++ b/Engine/Plugins/Runtime/OpenXR/Source/OpenXRHMD/Private/OpenXRHMDModule.cpp
@@ -456,6 +456,7 @@ bool FOpenXRHMDModule::GetOptionalExtensions(TArray<const ANSICHAR*>& OutExtensi
 #endif
 #ifdef XR_USE_GRAPHICS_API_VULKAN
 	OutExtensions.Add(XR_KHR_VULKAN_ENABLE_EXTENSION_NAME);
+	OutExtensions.Add(XR_KHR_VULKAN_SWAPCHAIN_FORMAT_LIST_EXTENSION_NAME);
 #endif
 	OutExtensions.Add(XR_KHR_COMPOSITION_LAYER_DEPTH_EXTENSION_NAME);
 	OutExtensions.Add(XR_VARJO_QUAD_VIEWS_EXTENSION_NAME);
diff --git a/Engine/Plugins/Runtime/OpenXR/Source/OpenXRHMD/Private/OpenXRHMD_Swapchain.cpp b/Engine/Plugins/Runtime/OpenXR/Source/OpenXRHMD/Private/OpenXRHMD_Swapchain.cpp
--- a/Engine/Plugins/Runtime/OpenXR/Source/OpenXRHMD/Private/OpenXRHMD_Swapchain.cpp
+++ b/Engine/Plugins/Runtime/OpenXR/Source/OpenXRHMD/Private/OpenXRHMD_Swapchain.cpp
@@ -189,7 +189,7 @@ uint8 FOpenXRSwapchain::GetNearestSupportedSwapchainFormat(XrSession InSession,
 	return FallbackFormat;
 }
 
-XrSwapchain FOpenXRSwapchain::CreateSwapchain(XrSession InSession, uint32 PlatformFormat, uint32 SizeX, uint32 SizeY, uint32 ArraySize, uint32 NumMips, uint32 NumSamples, ETextureCreateFlags CreateFlags)
+XrSwapchain FOpenXRSwapchain::CreateSwapchain(XrSession InSession, uint32 PlatformFormat, uint32 SizeX, uint32 SizeY, uint32 ArraySize, uint32 NumMips, uint32 NumSamples, ETextureCreateFlags CreateFlags, void* Next /*= nullptr*/)
 {
 	XrSwapchainUsageFlags Usage = 0;
 	if (!(CreateFlags & TexCreate_SRGB))
@@ -224,7 +224,7 @@ XrSwapchain FOpenXRSwapchain::CreateSwapchain(XrSession InSession, uint32 Platfo
 	XrSwapchain Swapchain;
 	XrSwapchainCreateInfo info;
 	info.type = XR_TYPE_SWAPCHAIN_CREATE_INFO;
-	info.next = nullptr;
+	info.next = Next;
 	info.createFlags = EnumHasAnyFlags(CreateFlags, TexCreate_Dynamic) ? 0 : XR_SWAPCHAIN_CREATE_STATIC_IMAGE_BIT;
 	info.usageFlags = Usage;
 	info.format = PlatformFormat;
@@ -395,7 +395,16 @@ FXRSwapChainPtr CreateSwapchain_Vulkan(XrSession InSession, uint8 Format, uint32
 		return nullptr;
 	}
 
-	XrSwapchain Swapchain = FOpenXRSwapchain::CreateSwapchain(InSession, ToPlatformFormat(Format), SizeX, SizeY, ArraySize, NumMips, NumSamples, CreateFlags);
+	// When we specify the mutable format flag we want to inform the runtime that we'll only use a
+	// linear and an sRGB view format to allow for efficiency optimizations.
+	TArray<VkFormat> ViewFormatList;
+	ViewFormatList.Add((VkFormat)ToPlatformFormat(Format)); // sRGB format
+	ViewFormatList.Add((VkFormat)GPixelFormats[Format].PlatformFormat); // linear format
+	XrVulkanSwapchainFormatListCreateInfoKHR FormatListInfo = { XR_TYPE_VULKAN_SWAPCHAIN_FORMAT_LIST_CREATE_INFO_KHR };
+	FormatListInfo.viewFormatCount = ViewFormatList.Num();
+	FormatListInfo.viewFormats = ViewFormatList.GetData();
+
+	XrSwapchain Swapchain = FOpenXRSwapchain::CreateSwapchain(InSession, ViewFormatList[0], SizeX, SizeY, ArraySize, NumMips, NumSamples, CreateFlags, &FormatListInfo);
 	if (!Swapchain)
 	{
 		return nullptr;
diff --git a/Engine/Plugins/Runtime/OpenXR/Source/OpenXRHMD/Private/OpenXRHMD_Swapchain.h b/Engine/Plugins/Runtime/OpenXR/Source/OpenXRHMD/Private/OpenXRHMD_Swapchain.h
--- a/Engine/Plugins/Runtime/OpenXR/Source/OpenXRHMD/Private/OpenXRHMD_Swapchain.h
+++ b/Engine/Plugins/Runtime/OpenXR/Source/OpenXRHMD/Private/OpenXRHMD_Swapchain.h
@@ -20,7 +20,7 @@ public:
 	virtual void ReleaseCurrentImage_RHIThread() override final;
 
 	XrSwapchain GetHandle() { return Handle; }
-	static XrSwapchain CreateSwapchain(XrSession InSession, uint32 PlatformFormat, uint32 SizeX, uint32 SizeY, uint32 ArraySize, uint32 NumMips, uint32 NumSamples, ETextureCreateFlags CreateFlags);
+	static XrSwapchain CreateSwapchain(XrSession InSession, uint32 PlatformFormat, uint32 SizeX, uint32 SizeY, uint32 ArraySize, uint32 NumMips, uint32 NumSamples, ETextureCreateFlags CreateFlags, void* Next = nullptr);
 	static uint8 GetNearestSupportedSwapchainFormat(XrSession InSession, uint8 RequestedFormat, TFunction<uint32(uint8)> ToPlatformFormat = nullptr);
 
 protected:
