OpenXR + Editor: disable max viewport resolution

By default, the Editor caps the viewport resolution to the monitor resolution. This doesn't make sense when we're using VRPIE or VREM.
We could also drive this from some VR-specific Editor code, but OpenXREditorModule is convenient.

diff --git a/Engine/Plugins/Runtime/OpenXR/Source/OpenXREditor/Private/OpenXREditorModule.cpp b/Engine/Plugins/Runtime/OpenXR/Source/OpenXREditor/Private/OpenXREditorModule.cpp
--- a/Engine/Plugins/Runtime/OpenXR/Source/OpenXREditor/Private/OpenXREditorModule.cpp
+++ b/Engine/Plugins/Runtime/OpenXR/Source/OpenXREditor/Private/OpenXREditorModule.cpp
@@ -7,6 +7,8 @@
 #include "ISettingsModule.h"
 #include "Modules/ModuleManager.h"
 
+#include "Editor/EditorPerformanceSettings.h"
+
 #define LOCTEXT_NAMESPACE "OpenXR"
 
 void FOpenXREditorModule::StartupModule()
@@ -24,6 +26,19 @@ void FOpenXREditorModule::StartupModule()
 			GetMutableDefault<UOpenXRInputSettings>()
 		);
 	}
+
+	UEditorPerformanceSettings* EditorPerformanceSettings = GetMutableDefault<UEditorPerformanceSettings>();
+	if (EditorPerformanceSettings->bOverrideMaxViewportRenderingResolution)
+	{
+		UE_LOG(LogTemp, Warning, TEXT("Existing value for UEditorPerformanceSettings::MaxViewportRenderingResolution will be overriden for OpenXR."));
+	}
+
+	UE_LOG(LogTemp, Log, TEXT("OpenXR ignores max viewport resolution in editor to support full HMD resolutions."));
+	EditorPerformanceSettings->bOverrideMaxViewportRenderingResolution = true;
+	EditorPerformanceSettings->MaxViewportRenderingResolution = 0;
+
+	FPropertyChangedEvent DisabledMaxResolutionEvent(EditorPerformanceSettings->GetClass()->FindPropertyByName(GET_MEMBER_NAME_CHECKED(UEditorPerformanceSettings, MaxViewportRenderingResolution)), EPropertyChangeType::ValueSet);
+	EditorPerformanceSettings->PostEditChangeProperty(DisabledMaxResolutionEvent);
 }
 
 void FOpenXREditorModule::ShutdownModule()
