XRScribe: Initial plugin infrastructure to record + playback OpenXR API calls

XRScribe is implemented as an IOpenXRExtensionPlugin. The initial implementation intercepts OpenXR API calls by standing in as the OpenXR loader, and returning function pointers to XRScribe-internal API intercepts.

The initial check-in can be enabled, will build, and expose the API intercepts, but doesn't do anything (besides indirectly disabling OpenXR).

diff --git a/Engine/Plugins/Runtime/OpenXR/Source/OpenXRHMD/Private/OpenXRHMDModule.cpp b/Engine/Plugins/Runtime/OpenXR/Source/OpenXRHMD/Private/OpenXRHMDModule.cpp
--- a/Engine/Plugins/Runtime/OpenXR/Source/OpenXRHMD/Private/OpenXRHMDModule.cpp
+++ b/Engine/Plugins/Runtime/OpenXR/Source/OpenXRHMD/Private/OpenXRHMDModule.cpp
@@ -528,6 +528,14 @@ bool FOpenXRHMDModule::InitInstance()
 		GetProcAddr = GetDefaultLoader();
 	}
 
+	for (IOpenXRExtensionPlugin* Plugin : ExtModules)
+	{
+		if (Plugin->InsertOpenXRAPILayer(GetProcAddr))
+		{
+			UE_LOG(LogHMD, Log, TEXT("IOpenXRExtensionPlugin API layer enabled: %s"), *Plugin->GetDisplayName());
+		}
+	}
+
 	if (!PreInitOpenXRCore(GetProcAddr))
 	{
 		UE_LOG(LogHMD, Log, TEXT("Failed to initialize core functions. Please check that you have a valid OpenXR runtime installed."));
diff --git a/Engine/Plugins/Runtime/OpenXR/Source/OpenXRHMD/Public/IOpenXRExtensionPlugin.h b/Engine/Plugins/Runtime/OpenXR/Source/OpenXRHMD/Public/IOpenXRExtensionPlugin.h
--- a/Engine/Plugins/Runtime/OpenXR/Source/OpenXRHMD/Public/IOpenXRExtensionPlugin.h
+++ b/Engine/Plugins/Runtime/OpenXR/Source/OpenXRHMD/Public/IOpenXRExtensionPlugin.h
@@ -142,13 +142,23 @@ public:
 	}
 
 	/**
-	* Optionally provide a custom loader for the OpenXR plugin.
+	* Optionally use a custom loader (via GetInstanceProcAddr) for the OpenXR plugin.
 	*/
 	virtual bool GetCustomLoader(PFN_xrGetInstanceProcAddr* OutGetProcAddr)
 	{
 		return false;
 	}
 
+	/**
+	* Experimental: Optionally hand off the loader/plugin GetInstanceProcAddr to an extension plugin
+	* to support API layering. Returns true if plugin is wrapping API. Layers can chain by using
+	* received GetProcAddr to hand off API calls.
+	*/
+	virtual bool InsertOpenXRAPILayer(PFN_xrGetInstanceProcAddr& InOutGetProcAddr)
+	{
+		return false;
+	}
+
 	/**
 	* Indicates that the device we're currently running does not support a spectator view.
 	* This will only be called once at initialization and should only return a result based for the current device the engine is running on.
