Add more ways to enable OpenXR validation.

diff --git a/Engine/Plugins/Runtime/OpenXR/Source/OpenXRHMD/Private/OpenXRHMDModule.cpp b/Engine/Plugins/Runtime/OpenXR/Source/OpenXRHMD/Private/OpenXRHMDModule.cpp
--- a/Engine/Plugins/Runtime/OpenXR/Source/OpenXRHMD/Private/OpenXRHMDModule.cpp
+++ b/Engine/Plugins/Runtime/OpenXR/Source/OpenXRHMD/Private/OpenXRHMDModule.cpp
@@ -638,11 +638,21 @@ bool FOpenXRHMDModule::InitInstance()
 	//   3. Copy the DLL from the build target at, for example, C:\OpenXR-SDK-Source-main\build\win64\src\api_layers\XrApiLayer_core_validation.dll to
 	//      somewhere in your system path (e.g. c:\windows\system32); the OpenXR loader currently doesn't use the path the json file is in (this is a bug)
 
-	const bool bEnableOpenXRValidationLayer = CVarEnableOpenXRValidationLayer.GetValueOnAnyThread() != 0;
+	const bool bEnableOpenXRValidationLayer = (CVarEnableOpenXRValidationLayer.GetValueOnAnyThread() != 0)
+		|| FParse::Param(FCommandLine::Get(), TEXT("openxrdebug"))
+		|| FParse::Param(FCommandLine::Get(), TEXT("openxrvalidation"));
 	TArray<const char*> Layers;
-	if (bEnableOpenXRValidationLayer && AvailableLayers.Contains("XR_APILAYER_LUNARG_core_validation"))
+	if (bEnableOpenXRValidationLayer)
 	{
-		Layers.Add("XR_APILAYER_LUNARG_core_validation");
+		if (AvailableLayers.Contains("XR_APILAYER_LUNARG_core_validation"))
+		{
+			UE_LOG(LogHMD, Display, TEXT("Running with OpenXR validation layers, performance might be degraded."));
+			Layers.Add("XR_APILAYER_LUNARG_core_validation");
+		}
+		else
+		{
+			UE_LOG(LogHMD, Error, TEXT("OpenXR validation was requested, but the validation layer isn't available. Request ignored."));
+		}
 	}
 
 	// Engine registration can be disabled via console var.
