Add support for XR_EXT_active_action_set_priority draft extension.

This allows us to properly support the priority number set on the input mapping context.

diff --git a/Engine/Plugins/Runtime/OpenXR/Source/OpenXRHMD/Private/Epic_openxr.h b/Engine/Plugins/Runtime/OpenXR/Source/OpenXRHMD/Private/Epic_openxr.h
--- a/Engine/Plugins/Runtime/OpenXR/Source/OpenXRHMD/Private/Epic_openxr.h
+++ b/Engine/Plugins/Runtime/OpenXR/Source/OpenXRHMD/Private/Epic_openxr.h
@@ -12,7 +12,9 @@
 #define XR_TYPE_EPIC 1000059050
 
 typedef enum XrStructureTypeEPIC {
-	XR_TYPE_SPACE_ACCELERATION_EPIC = XR_TYPE_EPIC + 0,
+	XR_TYPE_SPACE_ACCELERATION_EPIC = XR_TYPE_EPIC,
+	XR_TYPE_ACTIVE_ACTION_SET_PRIORITY_EXT = 1000373000,
+	XR_STRUCTURE_TYPE_EPIC_MAX_ENUM = 0x7FFFFFFF
 } XrStructureTypeEPIC;
 
 
@@ -34,4 +36,20 @@ typedef struct XrSpaceAccelerationEPIC {
 	XrSpaceAccelerationFlagsEPIC			accelerationFlags;
 	XrVector3f								linearAcceleration;  // m/s^2
 	XrVector3f							angularAcceleration; // rad/s^2
-} XrSpaceAccelerationEPIC;
\ No newline at end of file
+} XrSpaceAccelerationEPIC;
+
+#define XR_EXT_active_action_set_priority 1
+#define XR_EXT_active_action_set_priority_SPEC_VERSION 1
+#define XR_EXT_ACTIVE_ACTION_SET_PRIORITY_EXTENSION_NAME "XR_EXT_active_action_set_priority"
+typedef struct XrActiveActionSetPriorityEXT {
+	XrActionSet    actionSet;
+	uint32_t       priority;
+} XrActiveActionSetPriorityEXT;
+
+// XrActiveActionSetPrioritiesEXT extends XrActionsSyncInfo
+typedef struct XrActiveActionSetPrioritiesEXT {
+	XrStructureTypeEPIC              type;
+	void* XR_MAY_ALIAS               next;
+	uint32_t                         countActionSetPriorities;
+	XrActiveActionSetPriorityEXT*    actionSetPriorities;
+} XrActiveActionSetPrioritiesEXT;
diff --git a/Engine/Plugins/Runtime/OpenXR/Source/OpenXRHMD/Private/OpenXRHMDModule.cpp b/Engine/Plugins/Runtime/OpenXR/Source/OpenXRHMD/Private/OpenXRHMDModule.cpp
--- a/Engine/Plugins/Runtime/OpenXR/Source/OpenXRHMD/Private/OpenXRHMDModule.cpp
+++ b/Engine/Plugins/Runtime/OpenXR/Source/OpenXRHMD/Private/OpenXRHMDModule.cpp
@@ -8,6 +8,7 @@
 #include "IOpenXRARModule.h"
 #include "BuildSettings.h"
 #include "GeneralProjectSettings.h"
+#include "Epic_openxr.h"
 
 #if PLATFORM_ANDROID
 #include <android_native_app_glue.h>
@@ -462,12 +463,11 @@ bool FOpenXRHMDModule::GetOptionalExtensions(TArray<const ANSICHAR*>& OutExtensi
 	OutExtensions.Add(XR_VARJO_QUAD_VIEWS_EXTENSION_NAME);
 	OutExtensions.Add(XR_KHR_VISIBILITY_MASK_EXTENSION_NAME);
 	OutExtensions.Add(XR_KHR_BINDING_MODIFICATION_EXTENSION_NAME);
+	OutExtensions.Add(XR_EPIC_VIEW_CONFIGURATION_FOV_EXTENSION_NAME);
 
 	// Draft extension not yet provided in headers
 	OutExtensions.Add("XR_EXT_dpad_binding");
-
-	// A pre-draft extension.
-	OutExtensions.Add(XR_EPIC_VIEW_CONFIGURATION_FOV_EXTENSION_NAME);
+	OutExtensions.Add("XR_EXT_active_action_set_priority");
 
 	return true;
 }
diff --git a/Engine/Plugins/Runtime/OpenXR/Source/OpenXRHMD/Public/OpenXRCore.h b/Engine/Plugins/Runtime/OpenXR/Source/OpenXRHMD/Public/OpenXRCore.h
--- a/Engine/Plugins/Runtime/OpenXR/Source/OpenXRHMD/Public/OpenXRCore.h
+++ b/Engine/Plugins/Runtime/OpenXR/Source/OpenXRHMD/Public/OpenXRCore.h
@@ -124,6 +124,13 @@ FORCEINLINE XrExtent2Df ToXrExtent2D(FVector2D Vector, float Scale = 1.0f)
 	return XrExtent2Df{ (float)Vector.X / Scale, (float)Vector.Y / Scale };
 }
 
+FORCEINLINE uint32 ToXrPriority(int32 Priority)
+{
+	// Ensure negative priority numbers map to the lower half of the 32-bit range.
+	// We do this by casting to an unsigned int and then flipping the signed bit.
+	return (uint32)Priority ^ (1 << 31);
+}
+
 /** List all OpenXR global entry points used by Unreal. */
 #define ENUM_XR_ENTRYPOINTS_GLOBAL(EnumMacro) \
 	EnumMacro(PFN_xrEnumerateApiLayerProperties,xrEnumerateApiLayerProperties) \
diff --git a/Engine/Plugins/Runtime/OpenXR/Source/OpenXRInput/Private/OpenXRInput.cpp b/Engine/Plugins/Runtime/OpenXR/Source/OpenXRInput/Private/OpenXRInput.cpp
--- a/Engine/Plugins/Runtime/OpenXR/Source/OpenXRInput/Private/OpenXRInput.cpp
+++ b/Engine/Plugins/Runtime/OpenXR/Source/OpenXRInput/Private/OpenXRInput.cpp
@@ -8,6 +8,7 @@
 #include "GameFramework/InputSettings.h"
 #include "IOpenXRExtensionPlugin.h"
 #include "HeadMountedDisplayFunctionLibrary.h"
+#include "Epic_openxr.h"
 
 #include "EnhancedInputModule.h"
 #include "InputAction.h"
@@ -535,7 +536,7 @@ void FOpenXRInputPlugin::FOpenXRInput::BuildEnhancedActions(TMap<FString, FInter
 
 	for (const TPair<TObjectPtr<UInputMappingContext>, int32> MappingContext : MappableInputConfig->GetMappingContexts())
 	{
-		FOpenXRActionSet ActionSet(Instance, MappingContext.Key->GetFName(), MappingContext.Key->ContextDescription.ToString(), MappingContext.Value, MappingContext.Key);
+		FOpenXRActionSet ActionSet(Instance, MappingContext.Key->GetFName(), MappingContext.Key->ContextDescription.ToString(), ToXrPriority(MappingContext.Value), MappingContext.Key);
 		TMap<FName, int32> ActionMap;
 
 		for (const FEnhancedActionKeyMapping& Mapping : MappingContext.Key->GetMappings())
@@ -740,45 +741,56 @@ void FOpenXRInputPlugin::FOpenXRInput::SyncActions(XrSession Session)
 {
 	if (OpenXRHMD->IsFocused() && bActionsAttached)
 	{
-		TSet<XrActionSet> ActiveSet;
+		TMap<XrActionSet, int32> ActiveSet;
 		IEnhancedInputModule::Get().GetLibrary()->ForEachSubsystem([this, &ActiveSet](IEnhancedInputSubsystemInterface* Subsystem)
 			{
 				if (Subsystem)
 				{
 					for (const FOpenXRActionSet& ActionSet : ActionSets)
 					{
-						if (ActionSet.Object && Subsystem->HasMappingContext(ActionSet.Object))
+						int32 Priority = 0;
+						if (ActionSet.Object && Subsystem->HasMappingContext(ActionSet.Object, Priority))
 						{
-							ActiveSet.Add(ActionSet.Handle);
+							int32* PriorityPtr = ActiveSet.Find(ActionSet.Handle);
+							if (PriorityPtr)
+							{
+								*PriorityPtr = FMath::Max(*PriorityPtr, Priority);
+							}
+							else
+							{
+								ActiveSet.Add(ActionSet.Handle, Priority);
+							}
 						}
 					}
 				}
 			});
 
+		TArray<XrActiveActionSet> ActiveActionSets;
+		TArray<XrActiveActionSetPriorityEXT> ActivePriorities;
+		for (const auto& ActionSet : ActiveSet)
+		{
+			ActiveActionSets.Add(XrActiveActionSet{ ActionSet.Key, XR_NULL_PATH });
+			ActivePriorities.Add(XrActiveActionSetPriorityEXT{ ActionSet.Key, ToXrPriority(ActionSet.Value) });
+		}
+
 		// If legacy actions are enabled, all action sets are always active
 		if (!LegacyActions.IsEmpty())
 		{
 			for (const FOpenXRActionSet& ActionSet : ActionSets)
 			{
-				ActiveSet.Add(ActionSet.Handle);
+				ActiveActionSets.Add(XrActiveActionSet{ ActionSet.Handle, XR_NULL_PATH });
 			}
 		}
 
 		// The controller is always active
 		if (ControllerActionSet)
 		{
-			ActiveSet.Add(ControllerActionSet->Handle);
+			ActiveActionSets.Add(XrActiveActionSet{ ControllerActionSet->Handle, XR_NULL_PATH });
 		}
 
-		XrActionsSyncInfo SyncInfo;
-		SyncInfo.type = XR_TYPE_ACTIONS_SYNC_INFO;
-		SyncInfo.next = nullptr;
-
-		TArray<XrActiveActionSet> ActiveActionSets;
-		for (XrActionSet ActionSet : ActiveSet)
-		{
-			ActiveActionSets.Add(XrActiveActionSet{ ActionSet, XR_NULL_PATH });
-		}
+		XrActionsSyncInfo SyncInfo = { XR_TYPE_ACTIONS_SYNC_INFO };
+		XrActiveActionSetPrioritiesEXT PrioritiesInfo = { XR_TYPE_ACTIVE_ACTION_SET_PRIORITY_EXT };
+		SyncInfo.next = &PrioritiesInfo;
 
 		for (IOpenXRExtensionPlugin* Plugin : OpenXRHMD->GetExtensionPlugins())
 		{
@@ -787,6 +799,8 @@ void FOpenXRInputPlugin::FOpenXRInput::SyncActions(XrSession Session)
 			SyncInfo.next = Plugin->OnSyncActions(Session, SyncInfo.next);
 		}
 
+		PrioritiesInfo.countActionSetPriorities = ActivePriorities.Num();
+		PrioritiesInfo.actionSetPriorities = ActivePriorities.GetData();
 		SyncInfo.countActiveActionSets = ActiveActionSets.Num();
 		SyncInfo.activeActionSets = ActiveActionSets.GetData();
 
