OpenXRHMD API changes to support LiveLinkXR

diff --git a/Engine/Plugins/Runtime/OpenXR/Source/OpenXRHMD/Private/OpenXRHMD.cpp b/Engine/Plugins/Runtime/OpenXR/Source/OpenXRHMD/Private/OpenXRHMD.cpp
--- a/Engine/Plugins/Runtime/OpenXR/Source/OpenXRHMD/Private/OpenXRHMD.cpp
+++ b/Engine/Plugins/Runtime/OpenXR/Source/OpenXRHMD/Private/OpenXRHMD.cpp
@@ -1856,6 +1856,9 @@ void FOpenXRHMD::DestroySession()
 	// FlushRenderingCommands must be called outside of SessionLock since some rendering threads will also lock this mutex.
 	FlushRenderingCommands();
 
+	// Clear all the tracked devices
+	ResetActionDevices();
+
 	FWriteScopeLock SessionLock(SessionHandleMutex);
 
 	if (Session != XR_NULL_HANDLE)
@@ -1964,6 +1967,16 @@ XrPath FOpenXRHMD::GetTrackedDevicePath(const int32 DeviceId)
 	return XR_NULL_PATH;
 }
 
+XrSpace FOpenXRHMD::GetTrackedDeviceSpace(const int32 DeviceId)
+{
+	FReadScopeLock DeviceLock(DeviceMutex);
+	if (DeviceSpaces.IsValidIndex(DeviceId))
+	{
+		return DeviceSpaces[DeviceId].Space;
+	}
+	return XR_NULL_HANDLE;
+}
+
 XrTime FOpenXRHMD::GetDisplayTime() const
 {
 	const FPipelinedFrameState& PipelineState = GetPipelinedFrameStateForThread();
diff --git a/Engine/Plugins/Runtime/OpenXR/Source/OpenXRHMD/Private/OpenXRHMD.h b/Engine/Plugins/Runtime/OpenXR/Source/OpenXRHMD/Private/OpenXRHMD.h
--- a/Engine/Plugins/Runtime/OpenXR/Source/OpenXRHMD/Private/OpenXRHMD.h
+++ b/Engine/Plugins/Runtime/OpenXR/Source/OpenXRHMD/Private/OpenXRHMD.h
@@ -349,6 +349,7 @@ public:
 	OPENXRHMD_API int32 AddActionDevice(XrAction Action, XrPath Path);
 	OPENXRHMD_API void ResetActionDevices();
 	OPENXRHMD_API XrPath GetTrackedDevicePath(const int32 DeviceId);
+	OPENXRHMD_API XrSpace GetTrackedDeviceSpace(const int32 DeviceId);
 
 	OPENXRHMD_API bool IsExtensionEnabled(const FString& Name) const { return EnabledExtensions.Contains(Name); }
 	OPENXRHMD_API XrInstance GetInstance() { return Instance; }
