Fixed OpenXR projects packaged for Oculus Mobile crashing with a null pointer dereference due to a circular reference within the FVulkanDynamicRHI constructor.

Also altered IVulkanDynamicRHI::GetIVulkanDynamicRHI() to fail more gracefully with a check in the future when GDynamicRHI is not yet initialized.

diff --git a/Engine/Plugins/Runtime/OpenXR/Source/OpenXRHMD/Private/OpenXRHMD.cpp b/Engine/Plugins/Runtime/OpenXR/Source/OpenXRHMD/Private/OpenXRHMD.cpp
--- a/Engine/Plugins/Runtime/OpenXR/Source/OpenXRHMD/Private/OpenXRHMD.cpp
+++ b/Engine/Plugins/Runtime/OpenXR/Source/OpenXRHMD/Private/OpenXRHMD.cpp
@@ -100,8 +100,6 @@ namespace {
 bool FOpenXRHMD::FVulkanExtensions::GetVulkanInstanceExtensionsRequired(TArray<const ANSICHAR*>& Out)
 {
 #ifdef XR_USE_GRAPHICS_API_VULKAN
-	TArray<VkExtensionProperties> Properties = GetIVulkanDynamicRHI()->RHIGetAllInstanceExtensions();
-
 	if (Extensions.IsEmpty())
 	{
 		PFN_xrGetVulkanInstanceExtensionsKHR GetVulkanInstanceExtensionsKHR;
@@ -116,24 +114,7 @@ bool FOpenXRHMD::FVulkanExtensions::GetVulkanInstanceExtensionsRequired(TArray<c
 	ANSICHAR* Context = nullptr;
 	for (ANSICHAR* Tok = FCStringAnsi::Strtok(Extensions.GetData(), " ", &Context); Tok != nullptr; Tok = FCStringAnsi::Strtok(nullptr, " ", &Context))
 	{
-		bool ExtensionFound = false;
-		for (int32 PropertyIndex = 0; PropertyIndex < Properties.Num(); PropertyIndex++)
-		{
-			const char* PropertyExtensionName = Properties[PropertyIndex].extensionName;
-
-			if (!FCStringAnsi::Strcmp(PropertyExtensionName, Tok))
-			{
-				Out.Add(Tok);
-				ExtensionFound = true;
-				break;
-			}
-		}
-
-		if (!ExtensionFound)
-		{
-			UE_LOG(LogHMD, Log, TEXT("Missing required Vulkan instance extension %S."), Tok);
-			return false;
-		}
+		Out.Add(Tok);
 	}
 #endif
 	return true;
@@ -142,8 +123,6 @@ bool FOpenXRHMD::FVulkanExtensions::GetVulkanInstanceExtensionsRequired(TArray<c
 bool FOpenXRHMD::FVulkanExtensions::GetVulkanDeviceExtensionsRequired(VkPhysicalDevice_T *pPhysicalDevice, TArray<const ANSICHAR*>& Out)
 {
 #ifdef XR_USE_GRAPHICS_API_VULKAN
-	TArray<VkExtensionProperties> Properties = GetIVulkanDynamicRHI()->RHIGetAllDeviceExtensions((VkPhysicalDevice)pPhysicalDevice);
-
 	if (DeviceExtensions.IsEmpty())
 	{
 		PFN_xrGetVulkanDeviceExtensionsKHR GetVulkanDeviceExtensionsKHR;
@@ -158,24 +137,7 @@ bool FOpenXRHMD::FVulkanExtensions::GetVulkanDeviceExtensionsRequired(VkPhysical
 	ANSICHAR* Context = nullptr;
 	for (ANSICHAR* Tok = FCStringAnsi::Strtok(DeviceExtensions.GetData(), " ", &Context); Tok != nullptr; Tok = FCStringAnsi::Strtok(nullptr, " ", &Context))
 	{
-		bool ExtensionFound = false;
-		for (int32 PropertyIndex = 0; PropertyIndex < Properties.Num(); PropertyIndex++)
-		{
-			const char* PropertyExtensionName = Properties[PropertyIndex].extensionName;
-
-			if (!FCStringAnsi::Strcmp(PropertyExtensionName, Tok))
-			{
-				Out.Add(Tok);
-				ExtensionFound = true;
-				break;
-			}
-		}
-
-		if (!ExtensionFound)
-		{
-			UE_LOG(LogHMD, Log, TEXT("Missing required Vulkan device extension %S."), Tok);
-			return false;
-		}
+		Out.Add(Tok);
 	}
 #endif
 	return true;
diff --git a/Engine/Source/Runtime/VulkanRHI/Public/IVulkanDynamicRHI.h b/Engine/Source/Runtime/VulkanRHI/Public/IVulkanDynamicRHI.h
--- a/Engine/Source/Runtime/VulkanRHI/Public/IVulkanDynamicRHI.h
+++ b/Engine/Source/Runtime/VulkanRHI/Public/IVulkanDynamicRHI.h
@@ -85,6 +85,7 @@ struct IVulkanDynamicRHI : public FDynamicRHI
 
 inline IVulkanDynamicRHI* GetIVulkanDynamicRHI()
 {
+	checkf(GDynamicRHI, TEXT("Tried to fetch RHI too early"));
 	check(GDynamicRHI->GetInterfaceType() == ERHIInterfaceType::Vulkan);
 	return GetDynamicRHI<IVulkanDynamicRHI>();
 }
