LiveLinkXR API changes and refactors
- Refactor XrPath manipulation using a new FOpenXRPath utility class
- Allow XrSpace to be queried from OpenXRHMD
- Rename Special motion sources to the body tracking roles

diff --git a/Engine/Plugins/Runtime/OpenXR/Source/OpenXRHMD/Private/OpenXRCore.cpp b/Engine/Plugins/Runtime/OpenXR/Source/OpenXRHMD/Private/OpenXRCore.cpp
--- a/Engine/Plugins/Runtime/OpenXR/Source/OpenXRHMD/Private/OpenXRCore.cpp
+++ b/Engine/Plugins/Runtime/OpenXR/Source/OpenXRHMD/Private/OpenXRCore.cpp
@@ -2,6 +2,67 @@
 
 #include "OpenXRCore.h"
 #include "HeadMountedDisplayTypes.h"
+#include "IOpenXRHMDModule.h"
+
+FOpenXRPath::FOpenXRPath(XrPath InPath)
+	: Path(InPath)
+{
+}
+
+FOpenXRPath::FOpenXRPath(FName InName)
+	: Path(IOpenXRHMDModule::Get().ResolveNameToPath(InName))
+{
+}
+
+FOpenXRPath::FOpenXRPath(const char* PathString)
+	: Path(XR_NULL_PATH)
+{
+	XrInstance Instance = IOpenXRHMDModule::Get().GetInstance();
+	if (ensure(Instance != XR_NULL_HANDLE))
+	{
+		XrResult Result = xrStringToPath(Instance, PathString, &Path);
+		check(XR_SUCCEEDED(Result));
+	}
+}
+
+FOpenXRPath::FOpenXRPath(const FString& PathString)
+	: FOpenXRPath((ANSICHAR*)StringCast<ANSICHAR>(*PathString).Get())
+{
+}
+
+FString FOpenXRPath::ToString() const
+{
+	XrInstance Instance = IOpenXRHMDModule::Get().GetInstance();
+	if (!ensure(Instance != XR_NULL_HANDLE))
+	{
+		return FString();
+	}
+
+	uint32 PathCount = 0;
+	char PathChars[XR_MAX_PATH_LENGTH];
+	XrResult Result = xrPathToString(Instance, Path, XR_MAX_PATH_LENGTH, &PathCount, PathChars);
+	check(XR_SUCCEEDED(Result));
+	return XR_SUCCEEDED(Result) ? FString(PathCount, PathChars) : FString();
+}
+
+uint32 FOpenXRPath::GetStringLength() const
+{
+	XrInstance Instance = IOpenXRHMDModule::Get().GetInstance();
+	if (!ensure(Instance != XR_NULL_HANDLE))
+	{
+		return 0;
+	}
+
+	uint32 PathCount = 0;
+	XrResult Result = xrPathToString(Instance, Path, 0, &PathCount, nullptr);
+	check(XR_SUCCEEDED(Result));
+	return PathCount - 1;
+}
+
+FName FOpenXRPath::ToName() const
+{
+	return IOpenXRHMDModule::Get().ResolvePathToName(Path);
+}
 
 #define DEFINE_XR_ENTRYPOINTS(Type,Func) Type OpenXRDynamicAPI::Func = nullptr;
 ENUM_XR_ENTRYPOINTS_GLOBAL(DEFINE_XR_ENTRYPOINTS);
diff --git a/Engine/Plugins/Runtime/OpenXR/Source/OpenXRHMD/Private/OpenXRHMD.cpp b/Engine/Plugins/Runtime/OpenXR/Source/OpenXRHMD/Private/OpenXRHMD.cpp
--- a/Engine/Plugins/Runtime/OpenXR/Source/OpenXRHMD/Private/OpenXRHMD.cpp
+++ b/Engine/Plugins/Runtime/OpenXR/Source/OpenXRHMD/Private/OpenXRHMD.cpp
@@ -161,7 +161,7 @@ void FOpenXRHMD::GetMotionControllerData(UObject* WorldContext, const EControlle
 			if (XR_SUCCEEDED(xrGetCurrentInteractionProfile(Session, GetTrackedDevicePath(Devices[(int32)Hand]), &Profile)) &&
 				Profile.interactionProfile != XR_NULL_PATH)
 			{
-				XR_ENSURE(OpenXRPathToFName(Instance, Profile.interactionProfile, MotionControllerData.DeviceName));
+				MotionControllerData.DeviceName = FOpenXRPath(Profile.interactionProfile);
 			}
 		}
 	}
@@ -271,18 +271,13 @@ bool FOpenXRHMD::GetCurrentInteractionProfile(const EControllerHand Hand, FStrin
 			}
 			else
 			{				
-				XR_ENSURE(OpenXRPathToFString(Instance, Profile.interactionProfile, InteractionProfile));
+				InteractionProfile = FOpenXRPath(Profile.interactionProfile);
 				return true;
 			}
 		}
 		else
 		{
-			FString PathStr;
-			XrResult PathResult = OpenXRPathToFString(Instance, Path, PathStr);
-			if (!XR_SUCCEEDED(PathResult))
-			{
-				PathStr = FString::Printf(TEXT("xrPathToString returned %s"), OpenXRResultToString(Result));
-			}
+			FString PathStr = FOpenXRPath(Path);
 			UE_LOG(LogHMD, Warning, TEXT("GetCurrentInteractionProfile for %i (%s) failed because xrGetCurrentInteractionProfile failed with result %s."), Hand, *PathStr, OpenXRResultToString(Result));
 			return false;
 		}
@@ -1819,7 +1814,7 @@ void FOpenXRHMD::DestroySession()
 	FlushRenderingCommands();
 
 	// Clear all the tracked devices
-	ResetActionDevices();
+	ResetTrackedDevices();
 
 	FWriteScopeLock SessionLock(SessionHandleMutex);
 
@@ -1890,7 +1885,7 @@ void FOpenXRHMD::DestroySession()
 	}
 }
 
-int32 FOpenXRHMD::AddActionDevice(XrAction Action, XrPath Path)
+int32 FOpenXRHMD::AddTrackedDevice(XrAction Action, XrPath Path)
 {
 	FWriteScopeLock DeviceLock(DeviceMutex);
 
@@ -1908,7 +1903,7 @@ int32 FOpenXRHMD::AddActionDevice(XrAction Action, XrPath Path)
 	return DeviceId;
 }
 
-void FOpenXRHMD::ResetActionDevices()
+void FOpenXRHMD::ResetTrackedDevices()
 {
 	FWriteScopeLock DeviceLock(DeviceMutex);
 
diff --git a/Engine/Plugins/Runtime/OpenXR/Source/OpenXRHMD/Private/OpenXRHMD.h b/Engine/Plugins/Runtime/OpenXR/Source/OpenXRHMD/Private/OpenXRHMD.h
--- a/Engine/Plugins/Runtime/OpenXR/Source/OpenXRHMD/Private/OpenXRHMD.h
+++ b/Engine/Plugins/Runtime/OpenXR/Source/OpenXRHMD/Private/OpenXRHMD.h
@@ -346,8 +346,8 @@ public:
 	OPENXRHMD_API bool IsRunning() const;
 	OPENXRHMD_API bool IsFocused() const;
 
-	OPENXRHMD_API int32 AddActionDevice(XrAction Action, XrPath Path);
-	OPENXRHMD_API void ResetActionDevices();
+	OPENXRHMD_API int32 AddTrackedDevice(XrAction Action, XrPath Path);
+	OPENXRHMD_API void ResetTrackedDevices();
 	OPENXRHMD_API XrPath GetTrackedDevicePath(const int32 DeviceId);
 	OPENXRHMD_API XrSpace GetTrackedDeviceSpace(const int32 DeviceId);
 
diff --git a/Engine/Plugins/Runtime/OpenXR/Source/OpenXRHMD/Private/OpenXRHMDModule.cpp b/Engine/Plugins/Runtime/OpenXR/Source/OpenXRHMD/Private/OpenXRHMDModule.cpp
--- a/Engine/Plugins/Runtime/OpenXR/Source/OpenXRHMD/Private/OpenXRHMDModule.cpp
+++ b/Engine/Plugins/Runtime/OpenXR/Source/OpenXRHMD/Private/OpenXRHMDModule.cpp
@@ -755,3 +755,64 @@ bool FOpenXRHMDModule::InitSystem()
 
 	return true;
 }
+
+FName FOpenXRHMDModule::ResolvePathToName(XrPath Path)
+{
+	{
+		FReadScopeLock Lock(NameMutex);
+		FName* FoundName = PathToName.Find(Path);
+		if (FoundName)
+		{
+			// We've already previously resolved this XrPath to an FName
+			return *FoundName;
+		}
+	}
+
+	uint32 PathCount = 0;
+	char PathChars[XR_MAX_PATH_LENGTH];
+	XrResult Result = xrPathToString(Instance, Path, XR_MAX_PATH_LENGTH, &PathCount, PathChars);
+	check(Result);
+	if (Result == XR_SUCCESS)
+	{
+		// Resolve this XrPath to an FName and store it in the name map
+		FName Name(PathCount, PathChars);
+
+		FWriteScopeLock Lock(NameMutex);
+		PathToName.Add(Path, Name);
+		NameToPath.Add(Name, Path);
+		return Name;
+	}
+	else
+	{
+		return NAME_None;
+	}
+}
+
+XrPath FOpenXRHMDModule::ResolveNameToPath(FName Name)
+{
+	{
+		FReadScopeLock Lock(NameMutex);
+		XrPath* FoundPath = NameToPath.Find(Name);
+		if (FoundPath)
+		{
+			// We've already previously resolved this FName to an XrPath
+			return *FoundPath;
+		}
+	}
+
+	XrPath Path = XR_NULL_PATH;
+	FString PathString = Name.ToString();
+	XrResult Result = xrStringToPath(Instance, StringCast<ANSICHAR>(*PathString).Get(), &Path);
+	check(XR_SUCCEEDED(Result));
+	if (Result == XR_SUCCESS)
+	{
+		FWriteScopeLock Lock(NameMutex);
+		PathToName.Add(Path, Name);
+		NameToPath.Add(Name, Path);
+		return Path;
+	}
+	else
+	{
+		return XR_NULL_PATH;
+	}
+}
diff --git a/Engine/Plugins/Runtime/OpenXR/Source/OpenXRHMD/Private/OpenXRHMDModule.h b/Engine/Plugins/Runtime/OpenXR/Source/OpenXRHMD/Private/OpenXRHMDModule.h
--- a/Engine/Plugins/Runtime/OpenXR/Source/OpenXRHMD/Private/OpenXRHMDModule.h
+++ b/Engine/Plugins/Runtime/OpenXR/Source/OpenXRHMD/Private/OpenXRHMDModule.h
@@ -38,6 +38,11 @@ public:
 	virtual bool IsExtensionEnabled(const FString& Name) const override { return EnabledExtensions.Contains(Name); }
 	virtual bool IsLayerAvailable(const FString& Name) const override { return EnabledLayers.Contains(Name); }
 	virtual bool IsLayerEnabled(const FString& Name) const override { return EnabledLayers.Contains(Name); }
+	virtual XrInstance GetInstance() const override { return Instance; }
+	virtual XrSystemId GetSystemId() const override { return System; }
+
+	virtual FName ResolvePathToName(XrPath Path) override;
+	virtual XrPath ResolveNameToPath(FName Name) override;
 
 private:
 	void* LoaderHandle;
@@ -51,6 +56,11 @@ private:
 	TRefCountPtr<class FOpenXRRenderBridge> RenderBridge;
 	TSharedPtr< IHeadMountedDisplayVulkanExtensions, ESPMode::ThreadSafe > VulkanExtensions;
 
+	// We cache all attempts to convert between XrPath and FName to avoid costly resolves
+	FRWLock NameMutex;
+	TSortedMap<XrPath, FName> PathToName;
+	TSortedMap<FName, XrPath, FDefaultAllocator, FNameFastLess> NameToPath;
+
 	bool EnumerateExtensions();
 	bool EnumerateLayers();
 	bool InitRenderBridge();
diff --git a/Engine/Plugins/Runtime/OpenXR/Source/OpenXRHMD/Public/IOpenXRHMDModule.h b/Engine/Plugins/Runtime/OpenXR/Source/OpenXRHMD/Public/IOpenXRHMDModule.h
--- a/Engine/Plugins/Runtime/OpenXR/Source/OpenXRHMD/Public/IOpenXRHMDModule.h
+++ b/Engine/Plugins/Runtime/OpenXR/Source/OpenXRHMD/Public/IOpenXRHMDModule.h
@@ -6,6 +6,8 @@
 #include "Modules/ModuleManager.h"
 #include "IHeadMountedDisplayModule.h"
 
+#include <openxr/openxr.h>
+
 /**
  * The public interface to this module.  In most cases, this interface is only public to sibling modules 
  * within this plugin.
@@ -41,4 +43,10 @@ public:
 
 	virtual bool IsLayerAvailable(const FString& Name) const = 0;
 	virtual bool IsLayerEnabled(const FString& Name) const = 0;
+
+	virtual XrInstance GetInstance() const = 0;
+	virtual XrSystemId GetSystemId() const = 0;
+
+	virtual FName ResolvePathToName(XrPath Path) = 0;
+	virtual XrPath ResolveNameToPath(FName Name) = 0;
 };
diff --git a/Engine/Plugins/Runtime/OpenXR/Source/OpenXRHMD/Public/OpenXRCore.h b/Engine/Plugins/Runtime/OpenXR/Source/OpenXRHMD/Public/OpenXRCore.h
--- a/Engine/Plugins/Runtime/OpenXR/Source/OpenXRHMD/Public/OpenXRCore.h
+++ b/Engine/Plugins/Runtime/OpenXR/Source/OpenXRHMD/Public/OpenXRCore.h
@@ -131,6 +131,60 @@ FORCEINLINE uint32 ToXrPriority(int32 Priority)
 	return (uint32)Priority ^ (1 << 31);
 }
 
+/**
+ * XrPath wrapper with convenience functions
+ */
+class OPENXRHMD_API FOpenXRPath
+{
+public:
+	FOpenXRPath(XrPath InPath);
+
+	/**
+	 * Efficiently converts an FName to an XrPath
+	 */
+	FOpenXRPath(FName InName);
+
+	/**
+	 * Converts a string to an XrPath
+	 */
+	FOpenXRPath(const char* PathString);
+	FOpenXRPath(const FString& PathString);
+
+	/**
+	 * Converts an XrPath to a readable format
+	 *
+	 * @return String representation of the path
+	 */
+	FString ToString() const;
+
+	/**
+	 * Get the number of characters, excluding null-terminator, that ToString() would yield
+	 */
+	uint32 GetStringLength() const;
+
+	/**
+	 * Efficiently converts an XrPath to an FName
+	 *
+	 * @return FName representing the path
+	 */
+	FName ToName() const;
+
+	operator bool() const { return Path != XR_NULL_PATH; }
+	operator XrPath() const { return Path; }
+	operator FString() const { return ToString(); }
+	operator FName() const { return ToName(); }
+
+	/**
+	 * Operators to append another path ensuring the / character is used between them
+	 */
+	FOpenXRPath operator/(const char* Suffix) const { return FOpenXRPath(ToString() / Suffix); }
+	FOpenXRPath operator/(FString Suffix) const { return FOpenXRPath(ToString() / Suffix); }
+	FOpenXRPath operator/(FOpenXRPath Suffix) const { return FOpenXRPath(ToString() / Suffix); }
+
+private:
+	XrPath Path;
+};
+
 /** List all OpenXR global entry points used by Unreal. */
 #define ENUM_XR_ENTRYPOINTS_GLOBAL(EnumMacro) \
 	EnumMacro(PFN_xrEnumerateApiLayerProperties,xrEnumerateApiLayerProperties) \
@@ -213,34 +267,16 @@ bool PreInitOpenXRCore(PFN_xrGetInstanceProcAddr InGetProcAddr);
  */
 bool InitOpenXRCore(XrInstance Instance);
 
-FORCEINLINE XrResult OpenXRPathToFString(XrInstance Instance, XrPath Path, FString& OutString)
+FORCEINLINE void FilterActionName(const char* InActionName, char* OutActionName)
 {
-	uint32 PathCount = 0;
-	char PathChars[XR_MAX_PATH_LENGTH];
-	XrResult Result = xrPathToString(Instance, Path, XR_MAX_PATH_LENGTH, &PathCount, PathChars);
-	if (Result == XR_SUCCESS)
-	{
-		OutString = FString(PathCount, PathChars);
-	}
-	else
-	{
-		OutString = "";
-	}	
-	return Result;
-}
+	static_assert(XR_MAX_ACTION_NAME_SIZE == XR_MAX_ACTION_SET_NAME_SIZE);
 
-FORCEINLINE XrResult OpenXRPathToFName(XrInstance Instance, XrPath Path, FName& OutFName)
-{
-	uint32 PathCount = 0;
-	char PathChars[XR_MAX_PATH_LENGTH];
-	XrResult Result = xrPathToString(Instance, Path, XR_MAX_PATH_LENGTH, &PathCount, PathChars);
-	if (Result == XR_SUCCESS)
+	// Ensure the action name is a well-formed path
+	size_t i;
+	for (i = 0; i < XR_MAX_ACTION_NAME_SIZE - 1 && InActionName[i] != '\0'; i++)
 	{
-		OutFName = FName(PathCount, PathChars);
+		unsigned char c = InActionName[i];
+		OutActionName[i] = (c == ' ') ? '-' : isalnum(c) ? tolower(c) : '_';
 	}
-	else
-	{
-		OutFName = NAME_None;
-	}
-	return Result;
-}
\ No newline at end of file
+	OutActionName[i] = '\0';
+}
diff --git a/Engine/Plugins/Runtime/OpenXR/Source/OpenXRInput/Private/OpenXRInput.cpp b/Engine/Plugins/Runtime/OpenXR/Source/OpenXRInput/Private/OpenXRInput.cpp
--- a/Engine/Plugins/Runtime/OpenXR/Source/OpenXRInput/Private/OpenXRInput.cpp
+++ b/Engine/Plugins/Runtime/OpenXR/Source/OpenXRInput/Private/OpenXRInput.cpp
@@ -36,33 +36,6 @@ namespace OpenXRSourceNames
 	static const FName RightAim("RightAim");
 }
 
-FORCEINLINE XrPath GetPath(XrInstance Instance, const char* PathString)
-{
-	XrPath Path = XR_NULL_PATH;
-	XrResult Result = xrStringToPath(Instance, PathString, &Path);
-	check(XR_SUCCEEDED(Result));
-	return Path;
-}
-
-FORCEINLINE XrPath GetPath(XrInstance Instance, const FString& PathString)
-{
-	return GetPath(Instance, (ANSICHAR*)StringCast<ANSICHAR>(*PathString).Get());
-}
-
-FORCEINLINE void FilterActionName(const char* InActionName, char* OutActionName)
-{
-	static_assert(XR_MAX_ACTION_NAME_SIZE == XR_MAX_ACTION_SET_NAME_SIZE);
-
-	// Ensure the action name is a well-formed path
-	size_t i;
-	for (i = 0; i < XR_MAX_ACTION_NAME_SIZE - 1 && InActionName[i] != '\0'; i++)
-	{
-		unsigned char c = InActionName[i];
-		OutActionName[i] = (c == ' ') ? '-' : isalnum(c) ? tolower(c) : '_';
-	}
-	OutActionName[i] = '\0';
-}
-
 FORCEINLINE XrActionType ToActionType(EInputActionValueType ValueType)
 {
 	switch (ValueType)
@@ -225,12 +198,12 @@ FOpenXRInputPlugin::FOpenXRController::FOpenXRController(XrActionSet InActionSet
 	XR_ENSURE(xrCreateAction(ActionSet, &Info, &VibrationAction));
 }
 
-void FOpenXRInputPlugin::FOpenXRController::AddActionDevices(FOpenXRHMD* HMD)
+void FOpenXRInputPlugin::FOpenXRController::AddTrackedDevices(FOpenXRHMD* HMD)
 {
 	if (HMD)
 	{
-		GripDeviceId = HMD->AddActionDevice(GripAction, UserPath);
-		AimDeviceId = HMD->AddActionDevice(AimAction, UserPath);
+		GripDeviceId = HMD->AddTrackedDevice(GripAction, UserPath);
+		AimDeviceId = HMD->AddTrackedDevice(AimAction, UserPath);
 	}
 }
 
@@ -278,8 +251,8 @@ FOpenXRInputPlugin::FOpenXRInput::FOpenXRInput(FOpenXRHMD* HMD)
 		MotionSourceToControllerHandMap.Add(TEXT("EControllerHand::AnyHand"), EControllerHand::AnyHand);
 
 		// Generate a list of the sub-action paths so we can query the left/right hand individually
-		SubactionPaths.Add(GetPath(Instance, "/user/hand/left"));
-		SubactionPaths.Add(GetPath(Instance, "/user/hand/right"));
+		SubactionPaths.Add(FOpenXRPath("/user/hand/left"));
+		SubactionPaths.Add(FOpenXRPath("/user/hand/right"));
 
 		OpenXRHMD->SetInputModule(this);
 	}
@@ -341,9 +314,9 @@ bool FOpenXRInputPlugin::FOpenXRInput::BuildActions(XrSession Session)
 	// Create an engine action set for pose input and haptic output
 	ControllerActionSet = MakeUnique<FOpenXRActionSet>(Instance, "controllers", "Controllers", 0);
 
-	XrPath LeftHand = GetPath(Instance, "/user/hand/left");
-	XrPath RightHand = GetPath(Instance, "/user/hand/right");
-	XrPath Head = GetPath(Instance, "/user/head");
+	XrPath LeftHand = FOpenXRPath("/user/hand/left");
+	XrPath RightHand = FOpenXRPath("/user/hand/right");
+	XrPath Head = FOpenXRPath("/user/head");
 
 	// Controller poses
 	Controllers.Add(EControllerHand::Left, FOpenXRController(ControllerActionSet->Handle, LeftHand, "Left Controller"));
@@ -351,17 +324,17 @@ bool FOpenXRInputPlugin::FOpenXRInput::BuildActions(XrSession Session)
 	Controllers.Add(EControllerHand::HMD, FOpenXRController(ControllerActionSet->Handle, Head, "HMD"));
 
 	// Make OpenXRHMD aware of the controller action spaces
-	Controllers[EControllerHand::Left].AddActionDevices(OpenXRHMD);
-	Controllers[EControllerHand::Right].AddActionDevices(OpenXRHMD);
+	Controllers[EControllerHand::Left].AddTrackedDevices(OpenXRHMD);
+	Controllers[EControllerHand::Right].AddTrackedDevices(OpenXRHMD);
 
 	// Generate a map of all supported interaction profiles to store suggested bindings
 	TMap<FString, FInteractionProfile> Profiles;
-	Profiles.Add("SimpleController", FInteractionProfile(GetPath(Instance, "/interaction_profiles/khr/simple_controller"), true));
-	Profiles.Add("Vive", FInteractionProfile(GetPath(Instance, "/interaction_profiles/htc/vive_controller"), true));
-	Profiles.Add("MixedReality", FInteractionProfile(GetPath(Instance, "/interaction_profiles/microsoft/motion_controller"), true));
-	Profiles.Add("OculusGo", FInteractionProfile(GetPath(Instance, "/interaction_profiles/oculus/go_controller"), false));
-	Profiles.Add("OculusTouch", FInteractionProfile(GetPath(Instance, "/interaction_profiles/oculus/touch_controller"), true));
-	Profiles.Add("ValveIndex", FInteractionProfile(GetPath(Instance, "/interaction_profiles/valve/index_controller"), true));
+	Profiles.Add("SimpleController", FInteractionProfile(FOpenXRPath("/interaction_profiles/khr/simple_controller"), true));
+	Profiles.Add("Vive", FInteractionProfile(FOpenXRPath("/interaction_profiles/htc/vive_controller"), true));
+	Profiles.Add("MixedReality", FInteractionProfile(FOpenXRPath("/interaction_profiles/microsoft/motion_controller"), true));
+	Profiles.Add("OculusGo", FInteractionProfile(FOpenXRPath("/interaction_profiles/oculus/go_controller"), false));
+	Profiles.Add("OculusTouch", FInteractionProfile(FOpenXRPath("/interaction_profiles/oculus/touch_controller"), true));
+	Profiles.Add("ValveIndex", FInteractionProfile(FOpenXRPath("/interaction_profiles/valve/index_controller"), true));
 
 	// Query extension plugins for interaction profiles
 	for (IOpenXRExtensionPlugin* Plugin : OpenXRHMD->GetExtensionPlugins())
@@ -406,25 +379,25 @@ bool FOpenXRInputPlugin::FOpenXRInput::BuildActions(XrSession Session)
 		{
 			// Add the bindings for the controller pose and haptics
 			Profile.Bindings.Add(XrActionSuggestedBinding{
-				Controllers[EControllerHand::Left].GripAction, GetPath(Instance, "/user/hand/left/input/grip/pose")
+				Controllers[EControllerHand::Left].GripAction, FOpenXRPath("/user/hand/left/input/grip/pose")
 				});
 			Profile.Bindings.Add(XrActionSuggestedBinding{
-				Controllers[EControllerHand::Right].GripAction, GetPath(Instance, "/user/hand/right/input/grip/pose")
+				Controllers[EControllerHand::Right].GripAction, FOpenXRPath("/user/hand/right/input/grip/pose")
 				});
 			Profile.Bindings.Add(XrActionSuggestedBinding{
-				Controllers[EControllerHand::Left].AimAction, GetPath(Instance, "/user/hand/left/input/aim/pose")
+				Controllers[EControllerHand::Left].AimAction, FOpenXRPath("/user/hand/left/input/aim/pose")
 				});
 			Profile.Bindings.Add(XrActionSuggestedBinding{
-				Controllers[EControllerHand::Right].AimAction, GetPath(Instance, "/user/hand/right/input/aim/pose")
+				Controllers[EControllerHand::Right].AimAction, FOpenXRPath("/user/hand/right/input/aim/pose")
 				});
 
 			if (Profile.HasHaptics)
 			{
 				Profile.Bindings.Add(XrActionSuggestedBinding{
-					Controllers[EControllerHand::Left].VibrationAction, GetPath(Instance, "/user/hand/left/output/haptic")
+					Controllers[EControllerHand::Left].VibrationAction, FOpenXRPath("/user/hand/left/output/haptic")
 					});
 				Profile.Bindings.Add(XrActionSuggestedBinding{
-					Controllers[EControllerHand::Right].VibrationAction, GetPath(Instance, "/user/hand/right/output/haptic")
+					Controllers[EControllerHand::Right].VibrationAction, FOpenXRPath("/user/hand/right/output/haptic")
 					});
 			}
 
@@ -571,11 +544,6 @@ void FOpenXRInputPlugin::FOpenXRInput::BuildEnhancedActions(TMap<FString, FInter
 
 void FOpenXRInputPlugin::FOpenXRInput::DestroyActions()
 {
-	if (OpenXRHMD)
-	{
-		OpenXRHMD->ResetActionDevices();
-	}
-
 	// Destroying an action set will also destroy all actions in the set
 	for (const FOpenXRActionSet& ActionSet : ActionSets)
 	{
@@ -629,7 +597,7 @@ bool FOpenXRInputPlugin::FOpenXRInput::SuggestBindingForKey(TMap<FString, FInter
 
 	// Parse the key name into an OpenXR interaction profile path
 	FString Path = "/user/hand/" + Tokens[1].ToLower();
-	XrPath TopLevel = GetPath(Instance, Path);
+	XrPath TopLevel = FOpenXRPath(Path);
 
 	// Map this key to the correct subaction for this profile
 	// We'll use this later to retrieve binding modifiers
@@ -696,7 +664,7 @@ bool FOpenXRInputPlugin::FOpenXRInput::SuggestBindingForKey(TMap<FString, FInter
 	}
 
 	// Add the binding to the profile
-	Profile->Bindings.Add(XrActionSuggestedBinding{ Action.Handle, GetPath(Instance, Path) });
+	Profile->Bindings.Add(XrActionSuggestedBinding{ Action.Handle, FOpenXRPath(Path) });
 	return true;
 }
 
diff --git a/Engine/Plugins/Runtime/OpenXR/Source/OpenXRInput/Private/OpenXRInput.h b/Engine/Plugins/Runtime/OpenXR/Source/OpenXRInput/Private/OpenXRInput.h
--- a/Engine/Plugins/Runtime/OpenXR/Source/OpenXRInput/Private/OpenXRInput.h
+++ b/Engine/Plugins/Runtime/OpenXR/Source/OpenXRInput/Private/OpenXRInput.h
@@ -88,7 +88,7 @@ public:
 
 		FOpenXRController(XrActionSet InActionSet, XrPath InUserPath, const char* InName);
 
-		void AddActionDevices(FOpenXRHMD* HMD);
+		void AddTrackedDevices(FOpenXRHMD* HMD);
 	};
 
 	struct FInteractionProfile
diff --git a/Engine/Plugins/Runtime/Steam/SteamVR/Source/SteamVRInputDevice/Private/SteamVRInputTypes.h b/Engine/Plugins/Runtime/Steam/SteamVR/Source/SteamVRInputDevice/Private/SteamVRInputTypes.h
--- a/Engine/Plugins/Runtime/Steam/SteamVR/Source/SteamVRInputDevice/Private/SteamVRInputTypes.h
+++ b/Engine/Plugins/Runtime/Steam/SteamVR/Source/SteamVRInputDevice/Private/SteamVRInputTypes.h
@@ -158,9 +158,6 @@ namespace SteamVRInputDeviceConstants
 
 	/* Total number of supported tracked devices */
 	static const int32 MaxControllers = k_unMaxTrackedDeviceCount;
-
-	/* The maximum number of "Special" hand designations */
-	static const int32 MaxSpecialDesignations = ((int32)EControllerHand::Special_9 - (int32)EControllerHand::Special_1) + 1;
 }
 
 /* Buttons available on a SteamVR controller */
diff --git a/Engine/Source/Runtime/HeadMountedDisplay/Private/MotionTrackedDeviceFunctionLibrary.cpp b/Engine/Source/Runtime/HeadMountedDisplay/Private/MotionTrackedDeviceFunctionLibrary.cpp
--- a/Engine/Source/Runtime/HeadMountedDisplay/Private/MotionTrackedDeviceFunctionLibrary.cpp
+++ b/Engine/Source/Runtime/HeadMountedDisplay/Private/MotionTrackedDeviceFunctionLibrary.cpp
@@ -77,7 +77,7 @@ bool UMotionTrackedDeviceFunctionLibrary::EnableMotionTrackingOfDevice(int32 Pla
 
 bool UMotionTrackedDeviceFunctionLibrary::EnableMotionTrackingOfSource(int32 PlayerIndex, FName SourceName)
 {
-	EControllerHand Hand = EControllerHand::Special_11;
+	EControllerHand Hand = EControllerHand::Special;
 	if (FXRMotionControllerBase::GetHandEnumForSourceName(SourceName, Hand))
 	{
 		if (IModularFeatures::Get().IsModularFeatureAvailable(IMotionTrackingSystemManagement::GetModularFeatureName()))
@@ -112,7 +112,7 @@ void UMotionTrackedDeviceFunctionLibrary::DisableMotionTrackingOfDevice(int32 Pl
 
 void UMotionTrackedDeviceFunctionLibrary::DisableMotionTrackingOfSource(int32 PlayerIndex, FName SourceName)
 {
-	EControllerHand Hand = EControllerHand::Special_11;
+	EControllerHand Hand = EControllerHand::Special;
 	if (FXRMotionControllerBase::GetHandEnumForSourceName(SourceName, Hand))
 	{
 		if (IModularFeatures::Get().IsModularFeatureAvailable(IMotionTrackingSystemManagement::GetModularFeatureName()))
@@ -149,7 +149,7 @@ bool UMotionTrackedDeviceFunctionLibrary::IsMotionTrackingEnabledForDevice(int32
 
 bool UMotionTrackedDeviceFunctionLibrary::IsMotionTrackingEnabledForSource(int32 PlayerIndex, FName SourceName)
 {
-	EControllerHand Hand = EControllerHand::Special_11;
+	EControllerHand Hand = EControllerHand::Special;
 	if (FXRMotionControllerBase::GetHandEnumForSourceName(SourceName, Hand))
 	{
 		if (IModularFeatures::Get().IsModularFeatureAvailable(IMotionTrackingSystemManagement::GetModularFeatureName()))
diff --git a/Engine/Source/Runtime/HeadMountedDisplay/Private/XRAssetFunctionLibrary.cpp b/Engine/Source/Runtime/HeadMountedDisplay/Private/XRAssetFunctionLibrary.cpp
--- a/Engine/Source/Runtime/HeadMountedDisplay/Private/XRAssetFunctionLibrary.cpp
+++ b/Engine/Source/Runtime/HeadMountedDisplay/Private/XRAssetFunctionLibrary.cpp
@@ -71,7 +71,7 @@ static UPrimitiveComponent* XRAssetFunctionLibrary_Impl::AddNamedDeviceVisualiza
 {
 	UPrimitiveComponent* NewComponent = nullptr;
 
-	EControllerHand HandId = EControllerHand::Special_11;
+	EControllerHand HandId = EControllerHand::Special;
 	if (FXRMotionControllerBase::GetHandEnumForSourceName(DeviceName, HandId))
 	{
 		TArray<IXRSystemAssets*> XRAssetSystems = IModularFeatures::Get().GetModularFeatureImplementations<IXRSystemAssets>(IXRSystemAssets::GetModularFeatureName());
diff --git a/Engine/Source/Runtime/HeadMountedDisplay/Private/XRMotionControllerBase.cpp b/Engine/Source/Runtime/HeadMountedDisplay/Private/XRMotionControllerBase.cpp
--- a/Engine/Source/Runtime/HeadMountedDisplay/Private/XRMotionControllerBase.cpp
+++ b/Engine/Source/Runtime/HeadMountedDisplay/Private/XRMotionControllerBase.cpp
@@ -18,17 +18,18 @@ namespace XRMotionControllerBase_Impl
 		{ TEXT("Pad")                                }, // EControllerHand::Pad
 		{ TEXT("ExternalCamera")                     }, // EControllerHand::ExternalCamera
 		{ TEXT("Gun")                                }, // EControllerHand::Gun
-		{ TEXT("Special_1")                          }, // EControllerHand::Special_1
-		{ TEXT("Special_2")                          }, // EControllerHand::Special_2
-		{ TEXT("Special_3")                          }, // EControllerHand::Special_3
-		{ TEXT("Special_4")                          }, // EControllerHand::Special_4
-		{ TEXT("Special_5")                          }, // EControllerHand::Special_5
-		{ TEXT("Special_6")                          }, // EControllerHand::Special_6
-		{ TEXT("Special_7")                          }, // EControllerHand::Special_7
-		{ TEXT("Special_8")                          }, // EControllerHand::Special_8
-		{ TEXT("Special_9")                          }, // EControllerHand::Special_9
-		{ TEXT("Special_10")                         }, // EControllerHand::Special_10
-		{ TEXT("Special_11")                         }  // EControllerHand::Special_11
+		{ TEXT("HMD")                                }, // EControllerHand::HMD
+		{ TEXT("Chest")                              }, // EControllerHand::Chest
+		{ TEXT("LeftShoulder")                       }, // EControllerHand::LeftShoulder
+		{ TEXT("RightShoulder")                      }, // EControllerHand::RightShoulder
+		{ TEXT("LeftElbow")                          }, // EControllerHand::LeftElbow
+		{ TEXT("RightElbow")                         }, // EControllerHand::RightElbow
+		{ TEXT("Waist")                              }, // EControllerHand::Waist
+		{ TEXT("LeftKnee")                           }, // EControllerHand::LeftKnee
+		{ TEXT("RightKnee")                          }, // EControllerHand::RightKnee
+		{ TEXT("LeftFoot")                           }, // EControllerHand::LeftFoot
+		{ TEXT("RightFoot")                          }, // EControllerHand::RightFoot
+		{ TEXT("Special")                            }  // EControllerHand::Special
 	};
 }
 
@@ -114,17 +115,18 @@ bool FXRMotionControllerBase::GetHandEnumForSourceName(const FName Source, ECont
 		MotionSourceToEControllerHandMap.Add(TEXT("Pad"), EControllerHand::Pad);
 		MotionSourceToEControllerHandMap.Add(TEXT("ExternalCamera"), EControllerHand::ExternalCamera);
 		MotionSourceToEControllerHandMap.Add(TEXT("Gun"), EControllerHand::Gun);
-		MotionSourceToEControllerHandMap.Add(TEXT("Special_1"), EControllerHand::Special_1);
-		MotionSourceToEControllerHandMap.Add(TEXT("Special_2"), EControllerHand::Special_2);
-		MotionSourceToEControllerHandMap.Add(TEXT("Special_3"), EControllerHand::Special_3);
-		MotionSourceToEControllerHandMap.Add(TEXT("Special_4"), EControllerHand::Special_4);
-		MotionSourceToEControllerHandMap.Add(TEXT("Special_5"), EControllerHand::Special_5);
-		MotionSourceToEControllerHandMap.Add(TEXT("Special_6"), EControllerHand::Special_6);
-		MotionSourceToEControllerHandMap.Add(TEXT("Special_7"), EControllerHand::Special_7);
-		MotionSourceToEControllerHandMap.Add(TEXT("Special_8"), EControllerHand::Special_8);
-		MotionSourceToEControllerHandMap.Add(TEXT("Special_9"), EControllerHand::Special_9);
-		MotionSourceToEControllerHandMap.Add(TEXT("Special_10"), EControllerHand::Special_10);
-		MotionSourceToEControllerHandMap.Add(TEXT("Special_11"), EControllerHand::Special_11);
+		MotionSourceToEControllerHandMap.Add(TEXT("HMD"), EControllerHand::HMD);
+		MotionSourceToEControllerHandMap.Add(TEXT("Chest"), EControllerHand::Chest);
+		MotionSourceToEControllerHandMap.Add(TEXT("LeftShoulder"), EControllerHand::LeftShoulder);
+		MotionSourceToEControllerHandMap.Add(TEXT("RightShoulder"), EControllerHand::RightShoulder);
+		MotionSourceToEControllerHandMap.Add(TEXT("LeftElbow"), EControllerHand::LeftElbow);
+		MotionSourceToEControllerHandMap.Add(TEXT("RightElbow"), EControllerHand::RightElbow);
+		MotionSourceToEControllerHandMap.Add(TEXT("Waist"), EControllerHand::Waist);
+		MotionSourceToEControllerHandMap.Add(TEXT("LeftKnee"), EControllerHand::LeftKnee);
+		MotionSourceToEControllerHandMap.Add(TEXT("RightKnee"), EControllerHand::RightKnee);
+		MotionSourceToEControllerHandMap.Add(TEXT("LeftFoot"), EControllerHand::LeftFoot);
+		MotionSourceToEControllerHandMap.Add(TEXT("RightFoot"), EControllerHand::RightFoot);
+		MotionSourceToEControllerHandMap.Add(TEXT("Special"), EControllerHand::Special);
 		// EControllerHand enum names mapped to EControllerHand enum values
 		MotionSourceToEControllerHandMap.Add(TEXT("EControllerHand::Left"), EControllerHand::Left);
 		MotionSourceToEControllerHandMap.Add(TEXT("EControllerHand::Right"), EControllerHand::Right);
@@ -132,17 +134,18 @@ bool FXRMotionControllerBase::GetHandEnumForSourceName(const FName Source, ECont
 		MotionSourceToEControllerHandMap.Add(TEXT("EControllerHand::Pad"), EControllerHand::Pad);
 		MotionSourceToEControllerHandMap.Add(TEXT("EControllerHand::ExternalCamera"), EControllerHand::ExternalCamera);
 		MotionSourceToEControllerHandMap.Add(TEXT("EControllerHand::Gun"), EControllerHand::Gun);
-		MotionSourceToEControllerHandMap.Add(TEXT("EControllerHand::Special_1"), EControllerHand::Special_1);
-		MotionSourceToEControllerHandMap.Add(TEXT("EControllerHand::Special_2"), EControllerHand::Special_2);
-		MotionSourceToEControllerHandMap.Add(TEXT("EControllerHand::Special_3"), EControllerHand::Special_3);
-		MotionSourceToEControllerHandMap.Add(TEXT("EControllerHand::Special_4"), EControllerHand::Special_4);
-		MotionSourceToEControllerHandMap.Add(TEXT("EControllerHand::Special_5"), EControllerHand::Special_5);
-		MotionSourceToEControllerHandMap.Add(TEXT("EControllerHand::Special_6"), EControllerHand::Special_6);
-		MotionSourceToEControllerHandMap.Add(TEXT("EControllerHand::Special_7"), EControllerHand::Special_7);
-		MotionSourceToEControllerHandMap.Add(TEXT("EControllerHand::Special_8"), EControllerHand::Special_8);
-		MotionSourceToEControllerHandMap.Add(TEXT("EControllerHand::Special_9"), EControllerHand::Special_9);
-		MotionSourceToEControllerHandMap.Add(TEXT("EControllerHand::Special_10"), EControllerHand::Special_10);
-		MotionSourceToEControllerHandMap.Add(TEXT("EControllerHand::Special_11"), EControllerHand::Special_11);
+		MotionSourceToEControllerHandMap.Add(TEXT("EControllerHand::HMD"), EControllerHand::HMD);
+		MotionSourceToEControllerHandMap.Add(TEXT("EControllerHand::Chest"), EControllerHand::Chest);
+		MotionSourceToEControllerHandMap.Add(TEXT("EControllerHand::LeftShoulder"), EControllerHand::LeftShoulder);
+		MotionSourceToEControllerHandMap.Add(TEXT("EControllerHand::RightShoulder"), EControllerHand::RightShoulder);
+		MotionSourceToEControllerHandMap.Add(TEXT("EControllerHand::LeftElbow"), EControllerHand::LeftElbow);
+		MotionSourceToEControllerHandMap.Add(TEXT("EControllerHand::RightElbow"), EControllerHand::RightElbow);
+		MotionSourceToEControllerHandMap.Add(TEXT("EControllerHand::Waist"), EControllerHand::Waist);
+		MotionSourceToEControllerHandMap.Add(TEXT("EControllerHand::LeftKnee"), EControllerHand::LeftKnee);
+		MotionSourceToEControllerHandMap.Add(TEXT("EControllerHand::RightKnee"), EControllerHand::RightKnee);
+		MotionSourceToEControllerHandMap.Add(TEXT("EControllerHand::LeftFoot"), EControllerHand::LeftFoot);
+		MotionSourceToEControllerHandMap.Add(TEXT("EControllerHand::RightFoot"), EControllerHand::RightFoot);
+		MotionSourceToEControllerHandMap.Add(TEXT("EControllerHand::Special"), EControllerHand::Special);
 		// Newer source names that can usefully map to legacy EControllerHand values
 		MotionSourceToEControllerHandMap.Add(TEXT("LeftGrip"), EControllerHand::Left);
 		MotionSourceToEControllerHandMap.Add(TEXT("RightGrip"), EControllerHand::Right);
diff --git a/Engine/Source/Runtime/HeadMountedDisplay/Public/HeadMountedDisplayTypes.h b/Engine/Source/Runtime/HeadMountedDisplay/Public/HeadMountedDisplayTypes.h
--- a/Engine/Source/Runtime/HeadMountedDisplay/Public/HeadMountedDisplayTypes.h
+++ b/Engine/Source/Runtime/HeadMountedDisplay/Public/HeadMountedDisplayTypes.h
@@ -240,6 +240,8 @@ enum class EXRTrackedDeviceType : uint8
 	Controller,
 	/** Represents a static tracking reference device, such as a Lighthouse or tracking camera */
 	TrackingReference,
+	/** Represents trackers, such as a Vive tracker */
+	Tracker,
 	/** Misc. device types, for future expansion */
 	Other,
 	/** DeviceId is invalid */
diff --git a/Engine/Source/Runtime/InputCore/Classes/InputCoreTypes.h b/Engine/Source/Runtime/InputCore/Classes/InputCoreTypes.h
--- a/Engine/Source/Runtime/InputCore/Classes/InputCoreTypes.h
+++ b/Engine/Source/Runtime/InputCore/Classes/InputCoreTypes.h
@@ -23,17 +23,17 @@ enum class EControllerHand : uint8
 	ExternalCamera,
 	Gun,
 	HMD,
-	Special_1,
-	Special_2,
-	Special_3,
-	Special_4,
-	Special_5,
-	Special_6,
-	Special_7,
-	Special_8,
-	Special_9,
-	Special_10,
-	Special_11,
+	Chest,
+	LeftShoulder,
+	RightShoulder,
+	LeftElbow,
+	RightElbow,
+	Waist,
+	LeftKnee,
+	RightKnee,
+	LeftFoot,
+	RightFoot,
+	Special,
 
 	ControllerHand_Count UMETA(Hidden, DisplayName = "<INVALID>"),
 };
