OpenXR: Emulated world/tracker locked quad layers

If the OpenXR implementation does not support the expected amount of overlays, we restrict our usage of native layers to face-locked quads. We leverage the existing stereo layer compositing functionality in FDefaultStereoLayers to emulate world/tracker locked layers by compositing them directly in the background projection layer.

There are implementations that _do_ support the compliant number of layers, but using runtime layers significantly affects performance. This initial infrastructure allows apps to offload some performance implications by having FOpenXRHMD do some of the layer compositing in-engine.

diff --git a/Engine/Plugins/Runtime/OpenXR/Source/OpenXRHMD/Private/OpenXRHMD.cpp b/Engine/Plugins/Runtime/OpenXR/Source/OpenXRHMD/Private/OpenXRHMD.cpp
--- a/Engine/Plugins/Runtime/OpenXR/Source/OpenXRHMD/Private/OpenXRHMD.cpp
+++ b/Engine/Plugins/Runtime/OpenXR/Source/OpenXRHMD/Private/OpenXRHMD.cpp
@@ -34,6 +34,7 @@
 #include "HDRHelper.h"
 #include "Shader.h"
 #include "ScreenRendering.h"
+#include "DefaultStereoLayers.h"
 
 #if WITH_EDITOR
 #include "Editor/UnrealEd/Classes/Editor/EditorEngine.h"
@@ -1050,6 +1051,11 @@ void FOpenXRHMD::PreRenderView_RenderThread(FRDGBuilder& GraphBuilder, FSceneVie
 	check(IsInRenderingThread());
 }
 
+void FOpenXRHMD::PostRenderView_RenderThread(FRDGBuilder& GraphBuilder, FSceneView& InView)
+{
+	DrawEmulatedQuadLayers_RenderThread(GraphBuilder, InView);
+}
+
 void FOpenXRHMD::PreRenderViewFamily_RenderThread(FRDGBuilder& GraphBuilder, FSceneViewFamily& ViewFamily)
 {
 	check(IsInRenderingThread());
@@ -1204,6 +1210,10 @@ FOpenXRHMD::FOpenXRHMD(const FAutoRegister& AutoRegister, XrInstance InInstance,
 	bSupportsHandTracking = HandTrackingSystemProperties.supportsHandTracking == XR_TRUE;
 	SystemProperties.next = nullptr;
 
+	// Some runtimes aren't compliant with their number of layers supported.
+	// We support a fallback by emulating non-facelocked layers
+	bNativeWorldQuadLayerSupport = SystemProperties.graphicsProperties.maxLayerCount >= XR_MIN_COMPOSITION_LAYERS_SUPPORTED;
+
 	bSpaceAccellerationSupported = IsExtensionEnabled(XR_EPIC_SPACE_ACCELERATION_NAME);
 
 	static const auto CVarMobileMultiView = IConsoleManager::Get().FindTConsoleVariableDataInt(TEXT("vr.MobileMultiView"));
@@ -1835,6 +1845,8 @@ void FOpenXRHMD::DestroySession()
 			Layer.LeftEye.Swapchain.Reset();
 		});
 
+		NativeQuadLayers.Reset();
+
 		PipelinedLayerStateRendering.ColorSwapchain.Reset();
 		PipelinedLayerStateRendering.DepthSwapchain.Reset();
 		PipelinedLayerStateRendering.QuadSwapchains.Reset();
@@ -2168,44 +2180,9 @@ bool FOpenXRHMD::AllocateDepthTexture(uint32 Index, uint32 SizeX, uint32 SizeY,
 	return true;
 }
 
-void FOpenXRHMD::OnBeginRendering_RenderThread(FRHICommandListImmediate& RHICmdList, FSceneViewFamily& ViewFamily)
+void CreateNativeLayerSwapchain(FOpenXRLayer& Layer, TRefCountPtr<FOpenXRRenderBridge>& RenderBridge, XrSession Session)
 {
-	ensure(IsInRenderingThread());
-	if (!RenderBridge)
-	{
-		// Frame submission is not necessary in a headless session.
-		return;
-	}
-
-	for (IOpenXRExtensionPlugin* Module : ExtensionPlugins)
-	{
-		Module->OnBeginRendering_RenderThread(Session);
-	}
-
-	const float WorldToMeters = GetWorldToMetersScale();
-	const FTransform InvTrackingToWorld = GetTrackingToWorldTransform().Inverse();
-
-	for (int32 ViewIndex = 0; ViewIndex < ViewFamily.Views.Num(); ViewIndex++)
-	{
-		if (ViewFamily.Views[ViewIndex]->StereoPass == EStereoscopicPass::eSSP_FULL)
-		{
-			continue;
-		}
-
-		const XrView& View = PipelinedFrameStateRendering.Views[ViewIndex];
-		FTransform EyePose = ToFTransform(View.pose, WorldToMeters);
-
-		// Apply the base HMD pose to each eye pose, we will late update this pose for late update in another callback
-		FTransform BasePose(ViewFamily.Views[ViewIndex]->BaseHmdOrientation, ViewFamily.Views[ViewIndex]->BaseHmdLocation);
-		XrCompositionLayerProjectionView& Projection = PipelinedLayerStateRendering.ProjectionLayers[ViewIndex];
-		FTransform BasePoseTransform = EyePose * BasePose;
-		BasePoseTransform.NormalizeRotation();
-		Projection.pose = ToXrPose(BasePoseTransform, WorldToMeters);
-		Projection.fov = View.fov;
-	}
-
-	// Manage the swapchains in the stereo layers
-	auto CreateSwapchain = [this](FRHITexture2D* Texture, ETextureCreateFlags Flags)
+	auto CreateSwapchain = [&](FRHITexture2D* Texture, ETextureCreateFlags Flags)
 	{
 		return RenderBridge->CreateSwapchain(Session,
 			IStereoRenderTargetManager::GetStereoLayerPixelFormat(),
@@ -2217,49 +2194,92 @@ void FOpenXRHMD::OnBeginRendering_RenderThread(FRHICommandListImmediate& RHICmdL
 			Texture->GetFlags() | Flags | TexCreate_RenderTargetable,
 			Texture->GetClearBinding());
 	};
+
+	const ETextureCreateFlags Flags = Layer.Desc.Flags & IStereoLayers::LAYER_FLAG_TEX_CONTINUOUS_UPDATE ?
+		TexCreate_Dynamic | TexCreate_SRGB : TexCreate_SRGB;
+
+	if (Layer.NeedReallocateRightTexture())
+	{
+		FRHITexture2D* Texture = Layer.Desc.Texture->GetTexture2D();
+		Layer.RightEye.SetSwapchain(CreateSwapchain(Texture, Flags), Texture->GetSizeXY());
+	}
+
+	if (Layer.NeedReallocateLeftTexture())
+	{
+		FRHITexture2D* Texture = Layer.Desc.LeftTexture->GetTexture2D();
+		Layer.LeftEye.SetSwapchain(CreateSwapchain(Texture, Flags), Texture->GetSizeXY());
+	}
+}
+
+void FOpenXRHMD::SetupFrameQuadLayers_RenderThread(FRHICommandListImmediate& RHICmdList)
+{
+	ensure(IsInRenderingThread());
+
 	if (GetStereoLayersDirty())
 	{
+		// Go over the dirtied layers to bin them into either native or emulated
+		EmulatedSceneLayers.Reset();
+		NativeQuadLayers.Reset();
+
 		ForEachLayer([&](uint32 /* unused */, FOpenXRLayer& Layer)
 		{
-			if (Layer.Desc.IsVisible() && Layer.Desc.HasShape<FQuadLayer>())
+			// We use native layer for facelocked and when world-locked has proper support
+			if ((Layer.Desc.PositionType == ELayerType::FaceLocked) || bNativeWorldQuadLayerSupport)
 			{
-				const ETextureCreateFlags Flags = Layer.Desc.Flags & IStereoLayers::LAYER_FLAG_TEX_CONTINUOUS_UPDATE ?
-					TexCreate_Dynamic | TexCreate_SRGB : TexCreate_SRGB;
-
-				if (Layer.NeedReallocateRightTexture())
+				if (Layer.Desc.IsVisible() && Layer.Desc.HasShape<FQuadLayer>())
 				{
-					FRHITexture2D* Texture = Layer.Desc.Texture->GetTexture2D();
-					Layer.RightEye.SetSwapchain(CreateSwapchain(Texture, Flags), Texture->GetSizeXY());
+					CreateNativeLayerSwapchain(Layer, RenderBridge, Session);
+					NativeQuadLayers.Add(Layer);
 				}
-
-				if (Layer.NeedReallocateLeftTexture())
+				else
 				{
-					FRHITexture2D* Texture = Layer.Desc.LeftTexture->GetTexture2D();
-					Layer.LeftEye.SetSwapchain(CreateSwapchain(Texture, Flags), Texture->GetSizeXY());
+					// We retain references in FPipelinedLayerState to avoid premature destruction
+					Layer.RightEye.Swapchain.Reset();
+					Layer.LeftEye.Swapchain.Reset();
 				}
 			}
-			else
+			else if (Layer.Desc.IsVisible())
 			{
-				// We retain references in FPipelinedLayerState to avoid premature destruction
-				Layer.RightEye.Swapchain.Reset();
-				Layer.LeftEye.Swapchain.Reset();
+				EmulatedSceneLayers.Add(Layer.Desc);
 			}
 		});
+
+		auto LayerCompare = [](const auto& A, const auto& B)
+		{
+			FLayerDesc DescA, DescB;
+			if (GetLayerDescMember(A, DescA) && GetLayerDescMember(B, DescB))
+			{
+				if (DescA.Priority < DescB.Priority)
+				{
+					return true;
+				}
+				if (DescA.Priority > DescB.Priority)
+				{
+					return false;
+				}
+				return DescA.Id < DescB.Id;
+			}
+			return false;
+		};
+
+		EmulatedSceneLayers.Sort(LayerCompare);
+		NativeQuadLayers.Sort(LayerCompare);
 	}
 
-	// Gather all active quad layers for composition sorted by priority
-	TArray<FOpenXRLayer> StereoLayers;
-	CopySortedLayers(StereoLayers);
-	PipelinedLayerStateRendering.QuadLayers.Reset(StereoLayers.Num());
-	PipelinedLayerStateRendering.QuadSwapchains.Reset(StereoLayers.Num());
+	const FTransform InvTrackingToWorld = GetTrackingToWorldTransform().Inverse();
+	const float WorldToMeters = GetWorldToMetersScale();
+
+	// Set up our OpenXR info per native quad layer. Emulated layers have everything in FLayerDesc
+	PipelinedLayerStateRendering.QuadLayers.Reset(NativeQuadLayers.Num());
+	PipelinedLayerStateRendering.QuadSwapchains.Reset(NativeQuadLayers.Num());
 	{
 		FReadScopeLock DeviceLock(DeviceMutex);
 
-		for (const FOpenXRLayer& Layer : StereoLayers)
+		for (const FOpenXRLayer& Layer : NativeQuadLayers)
 		{
 			const bool bNoAlpha = Layer.Desc.Flags & IStereoLayers::LAYER_FLAG_TEX_NO_ALPHA_CHANNEL;
 			const bool bIsStereo = Layer.Desc.LeftTexture.IsValid();
-			FTransform PositionTransform = Layer.Desc.PositionType == ELayerType::WorldLocked ?
+			const FTransform PositionTransform = Layer.Desc.PositionType == ELayerType::WorldLocked ?
 				InvTrackingToWorld : FTransform::Identity;
 
 			XrCompositionLayerQuad Quad = { XR_TYPE_COMPOSITION_LAYER_QUAD, nullptr };
@@ -2276,7 +2296,7 @@ void FOpenXRHMD::OnBeginRendering_RenderThread(FRHICommandListImmediate& RHICmdL
 				if (Layer.RightEye.bUpdateTexture && bIsRunning)
 				{
 					FRHITexture2D* SrcTexture = Layer.Desc.Texture->GetTexture2D();
-					FIntRect DstRect(FIntPoint(0, 0), Layer.RightEye.SwapchainSize.IntPoint());
+					const FIntRect DstRect(FIntPoint(0, 0), Layer.RightEye.SwapchainSize.IntPoint());
 					CopyTexture_RenderThread(RHICmdList, SrcTexture, FIntRect(), Layer.RightEye.Swapchain, DstRect, false, bNoAlpha);
 				}
 
@@ -2294,7 +2314,7 @@ void FOpenXRHMD::OnBeginRendering_RenderThread(FRHICommandListImmediate& RHICmdL
 				if (Layer.LeftEye.bUpdateTexture && bIsRunning)
 				{
 					FRHITexture2D* SrcTexture = Layer.Desc.LeftTexture->GetTexture2D();
-					FIntRect DstRect(FIntPoint(0, 0), Layer.LeftEye.SwapchainSize.IntPoint());
+					const FIntRect DstRect(FIntPoint(0, 0), Layer.LeftEye.SwapchainSize.IntPoint());
 					CopyTexture_RenderThread(RHICmdList, SrcTexture, FIntRect(), Layer.LeftEye.Swapchain, DstRect, false, bNoAlpha);
 				}
 
@@ -2307,6 +2327,114 @@ void FOpenXRHMD::OnBeginRendering_RenderThread(FRHICommandListImmediate& RHICmdL
 			}
 		}
 	}
+}
+
+void FOpenXRHMD::DrawEmulatedQuadLayers_RenderThread(FRDGBuilder& GraphBuilder, const FSceneView& InView)
+{
+	check(IsInRenderingThread());
+
+	if (bNativeWorldQuadLayerSupport || !IStereoRendering::IsStereoEyeView(InView))
+	{
+		return;
+	}
+
+	// Partially borrowed from FDefaultStereoLayers
+	AddPass(GraphBuilder, RDG_EVENT_NAME("OpenXREmulatedLayerRender"), [this, &InView](FRHICommandListImmediate& RHICmdList)
+	{
+		FViewMatrices ModifiedViewMatrices = InView.ViewMatrices;
+		ModifiedViewMatrices.HackRemoveTemporalAAProjectionJitter();
+		const FMatrix& ProjectionMatrix = ModifiedViewMatrices.GetProjectionMatrix();
+		const FMatrix& ViewProjectionMatrix = ModifiedViewMatrices.GetViewProjectionMatrix();
+
+		// Calculate a view matrix that only adjusts for eye position, ignoring head position, orientation and world position.
+		FVector EyeShift;
+		FQuat EyeOrientation;
+		GetRelativeEyePose(IXRTrackingSystem::HMDDeviceId, InView.StereoViewIndex, EyeOrientation, EyeShift);
+
+		FMatrix EyeMatrix = FTranslationMatrix(-EyeShift) * FInverseRotationMatrix(EyeOrientation.Rotator()) * FMatrix(
+			FPlane(0, 0, 1, 0),
+			FPlane(1, 0, 0, 0),
+			FPlane(0, 1, 0, 0),
+			FPlane(0, 0, 0, 1));
+
+		FQuat HmdOrientation = FQuat::Identity;
+		FVector HmdLocation = FVector::ZeroVector;
+		GetCurrentPose(IXRTrackingSystem::HMDDeviceId, HmdOrientation, HmdLocation);
+
+		FMatrix TrackerMatrix = FTranslationMatrix(-HmdLocation) * FInverseRotationMatrix(HmdOrientation.Rotator()) * EyeMatrix;
+
+		FDefaultStereoLayers::FLayerRenderParams RenderParams{
+			InView.UnscaledViewRect, // Viewport
+			{
+				ViewProjectionMatrix,				// WorldLocked,
+				TrackerMatrix * ProjectionMatrix,	// TrackerLocked,
+				EyeMatrix * ProjectionMatrix		// FaceLocked
+			}
+		};
+
+		TArray<FRHITransitionInfo, TInlineAllocator<16>> Infos;
+		for (const FLayerDesc& Layer : EmulatedSceneLayers)
+		{
+			Infos.Add(FRHITransitionInfo(Layer.Texture, ERHIAccess::Unknown, ERHIAccess::SRVGraphics));
+		}
+		if (Infos.Num())
+		{
+			RHICmdList.Transition(Infos);
+		}
+
+		FTexture2DRHIRef RenderTarget = InView.Family->RenderTarget->GetRenderTargetTexture();
+
+		FRHIRenderPassInfo RPInfo(RenderTarget, ERenderTargetActions::Load_Store);
+		RHICmdList.BeginRenderPass(RPInfo, TEXT("EmulatedStereoLayerRender"));
+		RHICmdList.SetViewport((float)RenderParams.Viewport.Min.X, (float)RenderParams.Viewport.Min.Y, 0.0f, (float)RenderParams.Viewport.Max.X, (float)RenderParams.Viewport.Max.Y, 1.0f);
+
+		if (bSplashIsShown || !IsBackgroundLayerVisible())
+		{
+			DrawClearQuad(RHICmdList, FLinearColor::Black);
+		}
+
+		FDefaultStereoLayers::StereoLayerRender(RHICmdList, EmulatedSceneLayers, RenderParams);
+
+		RHICmdList.EndRenderPass();
+	});
+}
+
+void FOpenXRHMD::OnBeginRendering_RenderThread(FRHICommandListImmediate& RHICmdList, FSceneViewFamily& ViewFamily)
+{
+	ensure(IsInRenderingThread());
+	if (!RenderBridge)
+	{
+		// Frame submission is not necessary in a headless session.
+		return;
+	}
+
+	for (IOpenXRExtensionPlugin* Module : ExtensionPlugins)
+	{
+		Module->OnBeginRendering_RenderThread(Session);
+	}
+
+	const float WorldToMeters = GetWorldToMetersScale();
+
+	for (int32 ViewIndex = 0; ViewIndex < ViewFamily.Views.Num(); ViewIndex++)
+	{
+		if (ViewFamily.Views[ViewIndex]->StereoPass == EStereoscopicPass::eSSP_FULL)
+		{
+			continue;
+		}
+
+		const XrView& View = PipelinedFrameStateRendering.Views[ViewIndex];
+		FTransform EyePose = ToFTransform(View.pose, WorldToMeters);
+
+		// Apply the base HMD pose to each eye pose, we will late update this pose for late update in another callback
+		FTransform BasePose(ViewFamily.Views[ViewIndex]->BaseHmdOrientation, ViewFamily.Views[ViewIndex]->BaseHmdLocation);
+		XrCompositionLayerProjectionView& Projection = PipelinedLayerStateRendering.ProjectionLayers[ViewIndex];
+		FTransform BasePoseTransform = EyePose * BasePose;
+		BasePoseTransform.NormalizeRotation();
+		Projection.pose = ToXrPose(BasePoseTransform, WorldToMeters);
+		Projection.fov = View.fov;
+	}
+
+	SetupFrameQuadLayers_RenderThread(RHICmdList);
 
 #if !PLATFORM_HOLOLENS
 	if (bHiddenAreaMaskSupported && bNeedReBuildOcclusionMesh)
diff --git a/Engine/Plugins/Runtime/OpenXR/Source/OpenXRHMD/Private/OpenXRHMD.h b/Engine/Plugins/Runtime/OpenXR/Source/OpenXRHMD/Private/OpenXRHMD.h
--- a/Engine/Plugins/Runtime/OpenXR/Source/OpenXRHMD/Private/OpenXRHMD.h
+++ b/Engine/Plugins/Runtime/OpenXR/Source/OpenXRHMD/Private/OpenXRHMD.h
@@ -245,6 +245,9 @@ protected:
 	// Used with FCoreDelegates
 	void VRHeadsetRecenterDelegate();
 
+	void SetupFrameQuadLayers_RenderThread(FRHICommandListImmediate& RHICmdList);
+	void DrawEmulatedQuadLayers_RenderThread(FRDGBuilder& GraphBuilder, const FSceneView& InView);
+
 public:
 	/** IXRTrackingSystem interface */
 	virtual bool DoesSupportLateProjectionUpdate() const override { return true; }
@@ -296,6 +299,7 @@ public:
 	virtual void SetupView(FSceneViewFamily& InViewFamily, FSceneView& InView) override;
 	virtual void BeginRenderViewFamily(FSceneViewFamily& InViewFamily) override;
 	virtual void PreRenderView_RenderThread(FRDGBuilder& GraphBuilder, FSceneView& InView) override;
+	virtual void PostRenderView_RenderThread(FRDGBuilder& GraphBuilder, FSceneView& InView) override; // for non-face locked compositing
 	virtual void PreRenderViewFamily_RenderThread(FRDGBuilder& GraphBuilder, FSceneViewFamily& InViewFamily) override;
 	virtual void PostRenderViewFamily_RenderThread(FRDGBuilder& GraphBuilder, FSceneViewFamily& InViewFamily) override;
 
@@ -421,4 +425,8 @@ private:
 	bool					bUseCustomReferenceSpace;
 	FQuat					BaseOrientation;
 	FVector					BasePosition;
+
+	bool					bNativeWorldQuadLayerSupport;
+	TArray<IStereoLayers::FLayerDesc> EmulatedSceneLayers;
+	TArray<FOpenXRLayer>			  NativeQuadLayers;
 };
diff --git a/Engine/Source/Runtime/HeadMountedDisplay/Private/DefaultStereoLayers.cpp b/Engine/Source/Runtime/HeadMountedDisplay/Private/DefaultStereoLayers.cpp
--- a/Engine/Source/Runtime/HeadMountedDisplay/Private/DefaultStereoLayers.cpp
+++ b/Engine/Source/Runtime/HeadMountedDisplay/Private/DefaultStereoLayers.cpp
@@ -56,7 +56,9 @@ FDefaultStereoLayers::FDefaultStereoLayers(const FAutoRegister& AutoRegister, FH
 }
 
 //=============================================================================
-void FDefaultStereoLayers::StereoLayerRender(FRHICommandListImmediate& RHICmdList, const TArray<uint32> & LayersToRender, const FLayerRenderParams& RenderParams) const
+
+// static
+void FDefaultStereoLayers::StereoLayerRender(FRHICommandListImmediate& RHICmdList, const TArray<FLayerDesc>& LayersToRender, const FLayerRenderParams& RenderParams)
 {
 	check(IsInRenderingThread());
 	if (!LayersToRender.Num())
@@ -89,13 +91,12 @@ void FDefaultStereoLayers::StereoLayerRender(FRHICommandListImmediate& RHICmdLis
 	GraphicsPSOInit.PrimitiveType = PT_TriangleList;
 
 	// Force initialization of pipeline state on first iteration:
-	bool bLastWasOpaque = (RenderThreadLayers[LayersToRender[0]].Flags & LAYER_FLAG_TEX_NO_ALPHA_CHANNEL) == 0;
-	bool bLastWasExternal = (RenderThreadLayers[LayersToRender[0]].Flags & LAYER_FLAG_TEX_EXTERNAL) == 0;
+	bool bLastWasOpaque = (LayersToRender[0].Flags & LAYER_FLAG_TEX_NO_ALPHA_CHANNEL) == 0;
+	bool bLastWasExternal = (LayersToRender[0].Flags & LAYER_FLAG_TEX_EXTERNAL) == 0;
 
 	// For each layer
-	for (uint32 LayerIndex : LayersToRender)
+	for (const FLayerDesc& Layer : LayersToRender)
 	{
-		const FLayerDesc& Layer = RenderThreadLayers[LayerIndex];
 		check(Layer.IsVisible());
 		const bool bIsOpaque = (Layer.Flags & LAYER_FLAG_TEX_NO_ALPHA_CHANNEL) != 0;
 		const bool bIsExternal = (Layer.Flags & LAYER_FLAG_TEX_EXTERNAL) != 0;
@@ -175,33 +176,30 @@ void FDefaultStereoLayers::PreRenderViewFamily_RenderThread(FRDGBuilder& GraphBu
 	{
 		return;
 	}
-
-	CopyLayers(RenderThreadLayers);
-
+	
 	// Sort layers
 	SortedSceneLayers.Reset();
 	SortedOverlayLayers.Reset();
-	uint32 LayerCount = RenderThreadLayers.Num();
-	for (uint32 LayerIndex = 0; LayerIndex < LayerCount; ++LayerIndex)
+
+	ForEachLayer([&](uint32 /* unused */, const FLayerDesc& Layer)
 	{
-		const auto& Layer = RenderThreadLayers[LayerIndex];
 		if (!Layer.IsVisible())
 		{
-			continue;
+			return;
 		}
 		if (Layer.PositionType == ELayerType::FaceLocked)
 		{
-			SortedOverlayLayers.Add(LayerIndex);
+			SortedOverlayLayers.Add(Layer);
 		}
 		else
 		{
-			SortedSceneLayers.Add(LayerIndex);
+			SortedSceneLayers.Add(Layer);
 		}
-	}
+	});
 
-	auto SortLayersPredicate = [&](const uint32& A, const uint32& B)
+	auto SortLayersPredicate = [&](const FLayerDesc& A, const FLayerDesc& B)
 	{
-		return RenderThreadLayers[A].Priority < RenderThreadLayers[B].Priority;
+		return A.Priority < B.Priority;
 	};
 	SortedSceneLayers.Sort(SortLayersPredicate);
 	SortedOverlayLayers.Sort(SortLayersPredicate);
@@ -247,13 +245,13 @@ void FDefaultStereoLayers::PostRenderView_RenderThread(FRDGBuilder& GraphBuilder
 		};
 
 		TArray<FRHITransitionInfo, TInlineAllocator<16>> Infos;
-		for (uint32 LayerIndex : SortedSceneLayers)
+		for (const FLayerDesc& SceneLayer : SortedSceneLayers)
 		{
-			Infos.Add(FRHITransitionInfo(RenderThreadLayers[LayerIndex].Texture, ERHIAccess::Unknown, ERHIAccess::SRVGraphics));
+			Infos.Add(FRHITransitionInfo(SceneLayer.Texture, ERHIAccess::Unknown, ERHIAccess::SRVGraphics));
 		}
-		for (uint32 LayerIndex : SortedOverlayLayers)
+		for (const FLayerDesc& OverlayLayer : SortedOverlayLayers)
 		{
-			Infos.Add(FRHITransitionInfo(RenderThreadLayers[LayerIndex].Texture, ERHIAccess::Unknown, ERHIAccess::SRVGraphics));
+			Infos.Add(FRHITransitionInfo(OverlayLayer.Texture, ERHIAccess::Unknown, ERHIAccess::SRVGraphics));
 		}
 		if (Infos.Num())
 		{
diff --git a/Engine/Source/Runtime/HeadMountedDisplay/Public/DefaultStereoLayers.h b/Engine/Source/Runtime/HeadMountedDisplay/Public/DefaultStereoLayers.h
--- a/Engine/Source/Runtime/HeadMountedDisplay/Public/DefaultStereoLayers.h
+++ b/Engine/Source/Runtime/HeadMountedDisplay/Public/DefaultStereoLayers.h
@@ -29,6 +29,16 @@ public:
 		return false;
 	}
 
+	/** Experimental struct */
+	struct FLayerRenderParams
+	{
+		FIntRect Viewport;
+		FMatrix RenderMatrices[3];
+	};
+
+	/** Experimental method */
+	static void StereoLayerRender(FRHICommandListImmediate& RHICmdList, const TArray<FLayerDesc>& LayersToRender, const FLayerRenderParams& RenderParams);
+
 protected:
 	
 	/**
@@ -39,20 +49,11 @@ protected:
 		HmdTransform = InHmdTransform;
 	}
 
-	struct FLayerRenderParams
-	{
-		FIntRect Viewport;
-		FMatrix RenderMatrices[3];
-	};
-
-	void StereoLayerRender(FRHICommandListImmediate& RHICmdList, const TArray<uint32> & LayersToRender, const FLayerRenderParams& RenderParams) const;
-
 	FHeadMountedDisplayBase* HMDDevice;
 	FTransform HmdTransform;
 
-	TArray<FLayerDesc> RenderThreadLayers;
-	TArray<uint32> SortedSceneLayers;
-	TArray<uint32> SortedOverlayLayers;
+	TArray<FLayerDesc> SortedSceneLayers;
+	TArray<FLayerDesc> SortedOverlayLayers;
 
 	friend class FHeadMountedDisplayBase;
 };
\ No newline at end of file
