Fix off-by-one error when converting an OpenXR path to a string
* This resulted in a double null-terminated string
* Also fixed a invalid check

diff --git a/Engine/Plugins/Runtime/OpenXR/Source/OpenXRHMD/Private/OpenXRCore.cpp b/Engine/Plugins/Runtime/OpenXR/Source/OpenXRHMD/Private/OpenXRCore.cpp
--- a/Engine/Plugins/Runtime/OpenXR/Source/OpenXRHMD/Private/OpenXRCore.cpp
+++ b/Engine/Plugins/Runtime/OpenXR/Source/OpenXRHMD/Private/OpenXRCore.cpp
@@ -42,7 +42,7 @@ FString FOpenXRPath::ToString() const
 	char PathChars[XR_MAX_PATH_LENGTH];
 	XrResult Result = xrPathToString(Instance, Path, XR_MAX_PATH_LENGTH, &PathCount, PathChars);
 	check(XR_SUCCEEDED(Result));
-	return XR_SUCCEEDED(Result) ? FString(PathCount, PathChars) : FString();
+	return XR_SUCCEEDED(Result) ? FString(PathCount - 1, PathChars) : FString();
 }
 
 uint32 FOpenXRPath::GetStringLength() const
diff --git a/Engine/Plugins/Runtime/OpenXR/Source/OpenXRHMD/Private/OpenXRHMDModule.cpp b/Engine/Plugins/Runtime/OpenXR/Source/OpenXRHMD/Private/OpenXRHMDModule.cpp
--- a/Engine/Plugins/Runtime/OpenXR/Source/OpenXRHMD/Private/OpenXRHMDModule.cpp
+++ b/Engine/Plugins/Runtime/OpenXR/Source/OpenXRHMD/Private/OpenXRHMDModule.cpp
@@ -771,7 +771,7 @@ FName FOpenXRHMDModule::ResolvePathToName(XrPath Path)
 	uint32 PathCount = 0;
 	char PathChars[XR_MAX_PATH_LENGTH];
 	XrResult Result = xrPathToString(Instance, Path, XR_MAX_PATH_LENGTH, &PathCount, PathChars);
-	check(Result);
+	check(XR_SUCCEEDED(Result));
 	if (Result == XR_SUCCESS)
 	{
 		// Resolve this XrPath to an FName and store it in the name map
