OpenXR: Add pipelined support for NativePresent in FOpenXRRenderBridge

There are situations where OpenXR might want to manually request or skip native present via FOpenXRRenderBridge.
This state has to be pipelined since NeedsNativePresent() is called from the RenderThread, and Present() needs to return matching value in RHIThread.

diff --git a/Engine/Plugins/Runtime/OpenXR/Source/OpenXRHMD/Private/OpenXRHMD_RenderBridge.cpp b/Engine/Plugins/Runtime/OpenXR/Source/OpenXRHMD/Private/OpenXRHMD_RenderBridge.cpp
--- a/Engine/Plugins/Runtime/OpenXR/Source/OpenXRHMD/Private/OpenXRHMD_RenderBridge.cpp
+++ b/Engine/Plugins/Runtime/OpenXR/Source/OpenXRHMD/Private/OpenXRHMD_RenderBridge.cpp
@@ -12,7 +12,7 @@ bool FOpenXRRenderBridge::Present(int32& InOutSyncInterval)
 	if (OpenXRHMD)
 	{
 		OpenXRHMD->OnFinishRendering_RHIThread();
-		bNeedsNativePresent = !OpenXRHMD->IsStandaloneStereoOnlyDevice();
+		bNeedsNativePresent = !OpenXRHMD->IsStandaloneStereoOnlyDevice() && bNativePresent_RHIThread;
 	}
 
 	InOutSyncInterval = 0; // VSync off
diff --git a/Engine/Plugins/Runtime/OpenXR/Source/OpenXRHMD/Private/OpenXRHMD_RenderBridge.h b/Engine/Plugins/Runtime/OpenXR/Source/OpenXRHMD/Private/OpenXRHMD_RenderBridge.h
--- a/Engine/Plugins/Runtime/OpenXR/Source/OpenXRHMD/Private/OpenXRHMD_RenderBridge.h
+++ b/Engine/Plugins/Runtime/OpenXR/Source/OpenXRHMD/Private/OpenXRHMD_RenderBridge.h
@@ -48,15 +48,34 @@ public:
 	/** FRHICustomPresent */
 	virtual bool Present(int32& InOutSyncInterval) override;
 
+	// Called from RenderThread
+	virtual bool NeedsNativePresent() override
+	{
+		return bNativePresent_RenderThread;
+	}
+
 	virtual bool Support10BitSwapchain() const { return false; }
 
 	virtual bool HDRGetMetaDataForStereo(EDisplayOutputFormat& OutDisplayOutputFormat, EDisplayColorGamut& OutDisplayColorGamut, bool& OutbHDRSupported) { return false; }
 
+	virtual void ShouldNativePresent_RenderThread(bool bPresent)
+	{
+		bNativePresent_RenderThread = bPresent;
+	}
+
+	virtual void ShouldNativePresent_RHIThread(bool bPresent)
+	{
+		bNativePresent_RHIThread = bPresent;
+	}
+
 protected:
 	uint64 AdapterLuid;
+	bool bNativePresent_RenderThread = true;
+	bool bNativePresent_RHIThread = true;
+	
+	FOpenXRHMD* OpenXRHMD;
 
 private:
-	FOpenXRHMD* OpenXRHMD;
 };
 
 #ifdef XR_USE_GRAPHICS_API_D3D11
diff --git a/Engine/Source/Runtime/RHI/Public/RHIResources.h b/Engine/Source/Runtime/RHI/Public/RHIResources.h
--- a/Engine/Source/Runtime/RHI/Public/RHIResources.h
+++ b/Engine/Source/Runtime/RHI/Public/RHIResources.h
@@ -2980,7 +2980,7 @@ public:
 	// @param InOutSyncInterval - in out param, indicates if vsync is on (>0) or off (==0).
 	// @return	true if native Present should be also be performed; false otherwise. If it returns
 	// true, then InOutSyncInterval could be modified to switch between VSync/NoVSync for the normal 
-	// Present.  Must match value previously returned by NeedsNormalPresent for this frame.
+	// Present.  Must match value previously returned by NeedsNativePresent for this frame.
 	virtual bool Present(int32& InOutSyncInterval) = 0;
 
 	// Called from RHI thread after native Present has been called
