Prevent spamming xrStopHapticFeedback when no feedback is active

Also de-duplicate the haptic feedback code

diff --git a/Engine/Plugins/Runtime/OpenXR/Source/OpenXRInput/Private/OpenXRInput.cpp b/Engine/Plugins/Runtime/OpenXR/Source/OpenXRInput/Private/OpenXRInput.cpp
--- a/Engine/Plugins/Runtime/OpenXR/Source/OpenXRInput/Private/OpenXRInput.cpp
+++ b/Engine/Plugins/Runtime/OpenXR/Source/OpenXRInput/Private/OpenXRInput.cpp
@@ -172,6 +172,7 @@ FOpenXRInputPlugin::FOpenXRController::FOpenXRController(XrActionSet InActionSet
 	, VibrationAction(XR_NULL_HANDLE)
 	, GripDeviceId(-1)
 	, AimDeviceId(-1)
+	, bHapticActive(false)
 {
 	XrActionCreateInfo Info;
 	Info.type = XR_TYPE_ACTION_CREATE_INFO;
@@ -1254,37 +1255,35 @@ void FOpenXRInputPlugin::FOpenXRInput::SetHapticFeedbackValues(int32 ControllerI
 
 	if (ControllerId == 0)
 	{
+		FOpenXRController* Controller = nullptr;
 		if ((Hand == (int32)EControllerHand::Left || Hand == (int32)EControllerHand::AnyHand) && Controllers.Contains(EControllerHand::Left))
 		{
-			XrHapticActionInfo HapticActionInfo;
-			HapticActionInfo.type = XR_TYPE_HAPTIC_ACTION_INFO;
-			HapticActionInfo.next = nullptr;
-			HapticActionInfo.subactionPath = XR_NULL_PATH;
-			HapticActionInfo.action = Controllers[EControllerHand::Left].VibrationAction;
-			if (Values.HapticBuffer == nullptr && (Values.Amplitude <= 0.0f || Values.Frequency < XR_FREQUENCY_UNSPECIFIED))
-			{
-				XR_ENSURE(xrStopHapticFeedback(Session, &HapticActionInfo));
-			}
-			else
-			{				
-				FOpenXRExtensionChainStructPtrs ScopedExtensionChainStructs;
-				if (Values.HapticBuffer != nullptr)
-				{
-					OpenXRHMD->GetApplyHapticFeedbackAddChainStructsDelegate().Broadcast(&HapticValue, ScopedExtensionChainStructs, Values.HapticBuffer);
-				}
-				XR_ENSURE(xrApplyHapticFeedback(Session, &HapticActionInfo, (const XrHapticBaseHeader*)&HapticValue));
-			}
+			Controller = Controllers.Find(EControllerHand::Left);
 		}
 		if ((Hand == (int32)EControllerHand::Right || Hand == (int32)EControllerHand::AnyHand) && Controllers.Contains(EControllerHand::Right))
+		{
+			Controller = Controllers.Find(EControllerHand::Right);
+		}
+		if (Hand == (int32)EControllerHand::HMD)
+		{
+			Controller = Controllers.Find(EControllerHand::HMD);
+		}
+
+		if (Controller)
 		{
 			XrHapticActionInfo HapticActionInfo;
 			HapticActionInfo.type = XR_TYPE_HAPTIC_ACTION_INFO;
 			HapticActionInfo.next = nullptr;
 			HapticActionInfo.subactionPath = XR_NULL_PATH;
-			HapticActionInfo.action = Controllers[EControllerHand::Right].VibrationAction;
+			HapticActionInfo.action = Controller->VibrationAction;
+
 			if (Values.HapticBuffer == nullptr && (Values.Amplitude <= 0.0f || Values.Frequency < XR_FREQUENCY_UNSPECIFIED))
 			{
-				XR_ENSURE(xrStopHapticFeedback(Session, &HapticActionInfo));
+				if (Controller->bHapticActive)
+				{
+					XR_ENSURE(xrStopHapticFeedback(Session, &HapticActionInfo));
+				}
+				Controller->bHapticActive = false;
 			}
 			else
 			{
@@ -1294,23 +1293,8 @@ void FOpenXRInputPlugin::FOpenXRInput::SetHapticFeedbackValues(int32 ControllerI
 					OpenXRHMD->GetApplyHapticFeedbackAddChainStructsDelegate().Broadcast(&HapticValue, ScopedExtensionChainStructs, Values.HapticBuffer);
 				}
 				XR_ENSURE(xrApplyHapticFeedback(Session, &HapticActionInfo, (const XrHapticBaseHeader*)&HapticValue));
-			}
-		}
-		if (Hand == (int32)EControllerHand::HMD)
-		{
-			XrHapticActionInfo HapticActionInfo;
-			HapticActionInfo.type = XR_TYPE_HAPTIC_ACTION_INFO;
-			HapticActionInfo.next = nullptr;
-			HapticActionInfo.subactionPath = XR_NULL_PATH;
-			HapticActionInfo.action = Controllers[EControllerHand::HMD].VibrationAction;
-			
-			if (Values.HapticBuffer == nullptr && (Values.Amplitude <= 0.0f || Values.Frequency < XR_FREQUENCY_UNSPECIFIED))
-			{
-				XR_ENSURE(xrStopHapticFeedback(Session, &HapticActionInfo));
-			}
-			else
-			{
-				XR_ENSURE(xrApplyHapticFeedback(Session, &HapticActionInfo, (const XrHapticBaseHeader*)&HapticValue));
+
+				Controller->bHapticActive = true;
 			}
 		}
 	}
diff --git a/Engine/Plugins/Runtime/OpenXR/Source/OpenXRInput/Private/OpenXRInput.h b/Engine/Plugins/Runtime/OpenXR/Source/OpenXRInput/Private/OpenXRInput.h
--- a/Engine/Plugins/Runtime/OpenXR/Source/OpenXRInput/Private/OpenXRInput.h
+++ b/Engine/Plugins/Runtime/OpenXR/Source/OpenXRInput/Private/OpenXRInput.h
@@ -86,6 +86,8 @@ public:
 		int32			GripDeviceId;
 		int32			AimDeviceId;
 
+		bool			bHapticActive;
+
 		FOpenXRController(XrActionSet InActionSet, XrPath InUserPath, const char* InName);
 
 		void AddTrackedDevices(FOpenXRHMD* HMD);
