OpenXR: Additional clean up of XrFrameState usage

* During startup, we can call GetDisplayTime before xrWaitFrame has fetched a valid predictedDisplayTime. just return 0 (which callers can handle, and was previous behavior)
* Properly reset XrFrameState valid bit during session destruction. Was leading to issues during multiple PIE runs where the render thread thought the XrFrameState was still valid

diff --git a/Engine/Plugins/Runtime/OpenXR/Source/OpenXRHMD/Private/OpenXRHMD.cpp b/Engine/Plugins/Runtime/OpenXR/Source/OpenXRHMD/Private/OpenXRHMD.cpp
--- a/Engine/Plugins/Runtime/OpenXR/Source/OpenXRHMD/Private/OpenXRHMD.cpp
+++ b/Engine/Plugins/Runtime/OpenXR/Source/OpenXRHMD/Private/OpenXRHMD.cpp
@@ -1864,8 +1864,11 @@ void FOpenXRHMD::DestroySession()
 		bTrackingSpaceInvalid = true;
 
 		// Reset the frame state.
+		PipelinedFrameStateGame.bXrFrameStateUpdated = false;
 		PipelinedFrameStateGame.FrameState = XrFrameState{ XR_TYPE_FRAME_STATE };
+		PipelinedFrameStateRendering.bXrFrameStateUpdated = false;
 		PipelinedFrameStateRendering.FrameState = XrFrameState{ XR_TYPE_FRAME_STATE };
+		PipelinedFrameStateRHI.bXrFrameStateUpdated = false;
 		PipelinedFrameStateRHI.FrameState = XrFrameState{ XR_TYPE_FRAME_STATE };
 
 		// VRFocus must be reset so FWindowsApplication::PollGameDeviceState does not incorrectly short-circuit.
@@ -1949,8 +1952,7 @@ XrSpace FOpenXRHMD::GetTrackedDeviceSpace(const int32 DeviceId)
 XrTime FOpenXRHMD::GetDisplayTime() const
 {
 	const FPipelinedFrameState& PipelineState = GetPipelinedFrameStateForThread();
-	check(PipelineState.bXrFrameStateUpdated);
-	return PipelineState.FrameState.predictedDisplayTime;
+	return PipelineState.bXrFrameStateUpdated ? PipelineState.FrameState.predictedDisplayTime : 0;
 }
 
 XrSpace FOpenXRHMD::GetTrackingSpace() const
