OpenXR: Enhanced Input editor subsystem support.

diff --git a/Engine/Plugins/Runtime/OpenXR/Source/OpenXRInput/OpenXRInput.Build.cs b/Engine/Plugins/Runtime/OpenXR/Source/OpenXRInput/OpenXRInput.Build.cs
--- a/Engine/Plugins/Runtime/OpenXR/Source/OpenXRInput/OpenXRInput.Build.cs
+++ b/Engine/Plugins/Runtime/OpenXR/Source/OpenXRInput/OpenXRInput.Build.cs
@@ -41,7 +41,8 @@ namespace UnrealBuildTool.Rules
             {
 				PrivateDependencyModuleNames.Add("EditorFramework");
                 PrivateDependencyModuleNames.Add("UnrealEd");
-            }
+				PrivateDependencyModuleNames.Add("InputEditor");
+			}
         }
     }
 }
diff --git a/Engine/Plugins/Runtime/OpenXR/Source/OpenXRInput/Private/OpenXRInput.cpp b/Engine/Plugins/Runtime/OpenXR/Source/OpenXRInput/Private/OpenXRInput.cpp
--- a/Engine/Plugins/Runtime/OpenXR/Source/OpenXRInput/Private/OpenXRInput.cpp
+++ b/Engine/Plugins/Runtime/OpenXR/Source/OpenXRInput/Private/OpenXRInput.cpp
@@ -19,6 +19,7 @@
 #if WITH_EDITOR
 #include "Editor/EditorEngine.h"
 #include "Editor.h"
+#include "EnhancedInputEditorSubsystem.h"
 #endif
 
 #include <openxr/openxr.h>
@@ -705,7 +706,7 @@ void FOpenXRInputPlugin::FOpenXRInput::SyncActions(XrSession Session)
 	if (OpenXRHMD->IsFocused() && bActionsAttached)
 	{
 		TMap<XrActionSet, int32> ActiveSet;
-		IEnhancedInputModule::Get().GetLibrary()->ForEachSubsystem([this, &ActiveSet](IEnhancedInputSubsystemInterface* Subsystem)
+		auto GetActiveForSubsystem = [this, &ActiveSet](IEnhancedInputSubsystemInterface* Subsystem)
 			{
 				if (Subsystem)
 				{
@@ -726,7 +727,17 @@ void FOpenXRInputPlugin::FOpenXRInput::SyncActions(XrSession Session)
 						}
 					}
 				}
-			});
+			};
+
+		IEnhancedInputModule::Get().GetLibrary()->ForEachSubsystem(GetActiveForSubsystem);
+
+#if WITH_EDITOR
+		if (GEditor)
+		{
+			// UEnhancedInputLibrary::ForEachSubsystem only enumerates runtime subsystems.
+			GetActiveForSubsystem(GEditor->GetEditorSubsystem<UEnhancedInputEditorSubsystem>());
+		}
+#endif
 
 		TArray<XrActiveActionSet> ActiveActionSets;
 		TArray<XrActiveActionSetPriorityEXT> ActivePriorities;
@@ -957,13 +968,24 @@ void FOpenXRInputPlugin::FOpenXRInput::SendControllerEvents()
 			Action.Triggers.MultiFind(Key, Triggers, false);
 			TArray<TObjectPtr<UInputModifier>> Modifiers;
 			Action.Modifiers.MultiFind(Key, Modifiers, false);
-			IEnhancedInputModule::Get().GetLibrary()->ForEachSubsystem([InputAction, InputValue, Triggers, Modifiers](IEnhancedInputSubsystemInterface* Subsystem)
+
+			auto InjectSubsystemInput = [InputAction, InputValue, Triggers, Modifiers](IEnhancedInputSubsystemInterface* Subsystem)
 				{
 					if (Subsystem)
 					{
 						Subsystem->InjectInputForAction(InputAction, InputValue, Modifiers, Triggers);
 					}
-				});
+				};
+
+			IEnhancedInputModule::Get().GetLibrary()->ForEachSubsystem(InjectSubsystemInput);
+
+#if WITH_EDITOR
+			if (GEditor)
+			{
+				// UEnhancedInputLibrary::ForEachSubsystem only enumerates runtime subsystems.
+				InjectSubsystemInput(GEditor->GetEditorSubsystem<UEnhancedInputEditorSubsystem>());
+			}
+#endif
 		}
 	}
 }
