OpenXR: Disable depth swapchain size check

ScreenPercentage can cause the SceneTextures to 'mismatch' against the swapchain size.
However, the renderer handles these size deltas, so we don't need this check for now.

diff --git a/Engine/Plugins/Runtime/OpenXR/Source/OpenXRHMD/Private/OpenXRHMD.cpp b/Engine/Plugins/Runtime/OpenXR/Source/OpenXRHMD/Private/OpenXRHMD.cpp
--- a/Engine/Plugins/Runtime/OpenXR/Source/OpenXRHMD/Private/OpenXRHMD.cpp
+++ b/Engine/Plugins/Runtime/OpenXR/Source/OpenXRHMD/Private/OpenXRHMD.cpp
@@ -2175,7 +2175,6 @@ bool FOpenXRHMD::AllocateDepthTexture(uint32 Index, uint32 SizeX, uint32 SizeY,
 	// Flags might also change. We expect TexCreate_DepthStencilTargetable | TexCreate_ShaderResource | TexCreate_InputAttachmentRead from SceneTextures
 	check(EnumHasAllFlags(DepthSwapchainDesc.Flags, UnifiedCreateFlags));
 	check(DepthSwapchainDesc.Format == Format);
-	check(DepthSwapchainDesc.Extent.X == SizeX && DepthSwapchainDesc.Extent.Y == SizeY);
 	check(DepthSwapchainDesc.NumMips == FMath::Max(NumMips, 1u));
 	check(DepthSwapchainDesc.NumSamples == NumSamples);
 
