UE-136613 With editor tabs open (with additional world viewports?) setting "world to meters scale" has no effect.
-There are multiple unreal worlds in the editor, for world scale to work correctly in PIE we need to make sure we only pay attention to the primary PIE window.

diff --git a/Engine/Plugins/Runtime/OpenXR/Source/OpenXRHMD/Private/OpenXRHMD.cpp b/Engine/Plugins/Runtime/OpenXR/Source/OpenXRHMD/Private/OpenXRHMD.cpp
--- a/Engine/Plugins/Runtime/OpenXR/Source/OpenXRHMD/Private/OpenXRHMD.cpp
+++ b/Engine/Plugins/Runtime/OpenXR/Source/OpenXRHMD/Private/OpenXRHMD.cpp
@@ -37,6 +37,7 @@
 #include "DefaultStereoLayers.h"
 
 #if WITH_EDITOR
+#include "Editor.h"
 #include "Editor/EditorEngine.h"
 #endif
 
@@ -2630,17 +2631,25 @@ bool FOpenXRHMD::ReadNextEvent(XrEventDataBuffer* buffer)
 
 bool FOpenXRHMD::OnStartGameFrame(FWorldContext& WorldContext)
 {
+#if WITH_EDITOR
+	// In the editor there can be multiple worlds.  An editor world, pie worlds, other viewport worlds for editor pages.
+	// XR hardware can only be running with one of them.
+	if (GIsEditor && GEditor && GEditor->GetPIEWorldContext() != nullptr)
+	{
+		if (!WorldContext.bIsPrimaryPIEInstance)
+		{
+			return false;
+		}
+	}
+#endif // WITH_EDITOR
+
 	const AWorldSettings* const WorldSettings = WorldContext.World() ? WorldContext.World()->GetWorldSettings() : nullptr;
 	if (WorldSettings)
 	{
 		WorldToMetersScale = WorldSettings->WorldToMeters;
 	}
 
-	// Only refresh this based on the game world.  When remoting there is also an editor world, which we do not want to have affect the transform.
-	if (WorldContext.World()->IsGameWorld())
-	{
-		RefreshTrackingToWorldTransform(WorldContext);
-	}
+	RefreshTrackingToWorldTransform(WorldContext);
 
 	// Process all pending messages.
 	XrEventDataBuffer event;
