Add missing platform include to OpenXRPlatformRHI.h

diff --git a/Engine/Plugins/Runtime/OpenXR/Source/OpenXRHMD/Private/OpenXRPlatformRHI.h b/Engine/Plugins/Runtime/OpenXR/Source/OpenXRHMD/Private/OpenXRPlatformRHI.h
--- a/Engine/Plugins/Runtime/OpenXR/Source/OpenXRHMD/Private/OpenXRPlatformRHI.h
+++ b/Engine/Plugins/Runtime/OpenXR/Source/OpenXRHMD/Private/OpenXRPlatformRHI.h
@@ -2,6 +2,8 @@
 
 #pragma once
 
+#include "HAL/Platform.h"
+
 #if PLATFORM_WINDOWS || PLATFORM_HOLOLENS
 #define XR_USE_PLATFORM_WIN32		1
 #define XR_USE_GRAPHICS_API_D3D11	1
