UE-162569 Hand rendering can easily render the right hand 18 times, leading to unexpected performance drop
-Made FOpenXRHMD::GetMotionControllerData and FOpenXRHandTracking::GetAllKeypointStates return data only for EControllerHand::Left and ::Right rather than returning Right data for all enum values.  Now we return a struct with the valid flag set to false for the other, non-hand, enum values.
-The valid flag was already false when a hand is not tracked, so existing usages should be handling this case gracefully.

diff --git a/Engine/Plugins/Runtime/OpenXR/Source/OpenXRHMD/Private/OpenXRHMD.cpp b/Engine/Plugins/Runtime/OpenXR/Source/OpenXRHMD/Private/OpenXRHMD.cpp
--- a/Engine/Plugins/Runtime/OpenXR/Source/OpenXRHMD/Private/OpenXRHMD.cpp
+++ b/Engine/Plugins/Runtime/OpenXR/Source/OpenXRHMD/Private/OpenXRHMD.cpp
@@ -152,6 +152,7 @@ void FOpenXRHMD::GetMotionControllerData(UObject* WorldContext, const EControlle
 	MotionControllerData.DeviceVisualType = EXRVisualType::Controller;
 	MotionControllerData.TrackingStatus = ETrackingStatus::NotTracked;
 	MotionControllerData.HandIndex = Hand;
+	MotionControllerData.bValid = false;
 
 	TArray<int32> Devices;
 	if (EnumerateTrackedDevices(Devices, EXRTrackedDeviceType::Controller) && Devices.IsValidIndex((int32)Hand))
@@ -180,53 +181,57 @@ void FOpenXRHMD::GetMotionControllerData(UObject* WorldContext, const EControlle
 		}
 	}
 
-	FName MotionControllerName("OpenXR");
-	TArray<IMotionController*> MotionControllers = IModularFeatures::Get().GetModularFeatureImplementations<IMotionController>(IMotionController::GetModularFeatureName());
-	IMotionController* MotionController = nullptr;
-	for (auto Itr : MotionControllers)
+	if ((Hand == EControllerHand::Left) || (Hand == EControllerHand::Right))
 	{
-		if (Itr->GetMotionControllerDeviceTypeName() == MotionControllerName)
+		FName MotionControllerName("OpenXR");
+		TArray<IMotionController*> MotionControllers = IModularFeatures::Get().GetModularFeatureImplementations<IMotionController>(IMotionController::GetModularFeatureName());
+		IMotionController* MotionController = nullptr;
+		for (auto Itr : MotionControllers)
 		{
-			MotionController = Itr;
-			break;
+			if (Itr->GetMotionControllerDeviceTypeName() == MotionControllerName)
+			{
+				MotionController = Itr;
+				break;
+			}
 		}
-	}
 
-	if (MotionController)
-	{
-		const float WorldToMeters = GetWorldToMetersScale();
-
-		bool bSuccess = false;
-		FVector Position = FVector::ZeroVector;
-		FRotator Rotation = FRotator::ZeroRotator;
-		FTransform trackingToWorld = GetTrackingToWorldTransform();
-		FName AimSource = Hand == EControllerHand::Left ? FName("LeftAim") : FName("RightAim");
-		bSuccess = MotionController->GetControllerOrientationAndPosition(0, AimSource, Rotation, Position, WorldToMeters);
-		if (bSuccess)
+		if (MotionController)
 		{
-			MotionControllerData.AimPosition = trackingToWorld.TransformPosition(Position);
-			MotionControllerData.AimRotation = trackingToWorld.TransformRotation(FQuat(Rotation));
-		}
-		MotionControllerData.bValid |= bSuccess;
+			const float WorldToMeters = GetWorldToMetersScale();
 
-		FName GripSource = Hand == EControllerHand::Left ? FName("LeftGrip") : FName("RightGrip");
-		bSuccess = MotionController->GetControllerOrientationAndPosition(0, GripSource, Rotation, Position, WorldToMeters);
-		if (bSuccess)
-		{
-			MotionControllerData.GripPosition = trackingToWorld.TransformPosition(Position);
-			MotionControllerData.GripRotation = trackingToWorld.TransformRotation(FQuat(Rotation));
+			bool bSuccess = false;
+			FVector Position = FVector::ZeroVector;
+			FRotator Rotation = FRotator::ZeroRotator;
+			FTransform trackingToWorld = GetTrackingToWorldTransform();
+			FName AimSource = Hand == EControllerHand::Left ? FName("LeftAim") : FName("RightAim");
+			bSuccess = MotionController->GetControllerOrientationAndPosition(0, AimSource, Rotation, Position, WorldToMeters);
+			if (bSuccess)
+			{
+				MotionControllerData.AimPosition = trackingToWorld.TransformPosition(Position);
+				MotionControllerData.AimRotation = trackingToWorld.TransformRotation(FQuat(Rotation));
+			}
+			MotionControllerData.bValid |= bSuccess;
+
+			FName GripSource = Hand == EControllerHand::Left ? FName("LeftGrip") : FName("RightGrip");
+			bSuccess = MotionController->GetControllerOrientationAndPosition(0, GripSource, Rotation, Position, WorldToMeters);
+			if (bSuccess)
+			{
+				MotionControllerData.GripPosition = trackingToWorld.TransformPosition(Position);
+				MotionControllerData.GripRotation = trackingToWorld.TransformRotation(FQuat(Rotation));
+			}
+			MotionControllerData.bValid |= bSuccess;
+
+			MotionControllerData.TrackingStatus = MotionController->GetControllerTrackingStatus(0, GripSource);
 		}
-		MotionControllerData.bValid |= bSuccess;
 
-		MotionControllerData.TrackingStatus = MotionController->GetControllerTrackingStatus(0, GripSource);
-	}
 
-	if (HandTracker && HandTracker->IsHandTrackingStateValid())
-	{
-		MotionControllerData.DeviceVisualType = EXRVisualType::Hand;
+		if (HandTracker && HandTracker->IsHandTrackingStateValid())
+		{
+			MotionControllerData.DeviceVisualType = EXRVisualType::Hand;
 
-		MotionControllerData.bValid = HandTracker->GetAllKeypointStates(Hand, MotionControllerData.HandKeyPositions, MotionControllerData.HandKeyRotations, MotionControllerData.HandKeyRadii);
-		check(!MotionControllerData.bValid || (MotionControllerData.HandKeyPositions.Num() == EHandKeypointCount && MotionControllerData.HandKeyRotations.Num() == EHandKeypointCount && MotionControllerData.HandKeyRadii.Num() == EHandKeypointCount));
+			MotionControllerData.bValid = HandTracker->GetAllKeypointStates(Hand, MotionControllerData.HandKeyPositions, MotionControllerData.HandKeyRotations, MotionControllerData.HandKeyRadii);
+			check(!MotionControllerData.bValid || (MotionControllerData.HandKeyPositions.Num() == EHandKeypointCount && MotionControllerData.HandKeyRotations.Num() == EHandKeypointCount && MotionControllerData.HandKeyRadii.Num() == EHandKeypointCount));
+		}	
 	}
 
 	//TODO: this is reportedly a wmr specific convenience function for rapid prototyping.  Not sure it is useful for openxr.
diff --git a/Engine/Plugins/Runtime/OpenXRHandTracking/Source/OpenXRHandTracking/Private/OpenXRHandTracking.cpp b/Engine/Plugins/Runtime/OpenXRHandTracking/Source/OpenXRHandTracking/Private/OpenXRHandTracking.cpp
--- a/Engine/Plugins/Runtime/OpenXRHandTracking/Source/OpenXRHandTracking/Private/OpenXRHandTracking.cpp
+++ b/Engine/Plugins/Runtime/OpenXRHandTracking/Source/OpenXRHandTracking/Private/OpenXRHandTracking.cpp
@@ -553,6 +553,11 @@ bool FOpenXRHandTracking::GetAllKeypointStates(EControllerHand Hand, TArray<FVec
 		return false;
 	}
 
+	if (Hand != EControllerHand::Left && Hand != EControllerHand::Right)
+	{
+		return false;
+	}
+
 	const FOpenXRHandTracking::FHandState& HandState = (Hand == EControllerHand::Left) ? GetLeftHandState() : GetRightHandState();
 
 	if (!HandState.ReceivedJointPoses)
