Set update flag for NativeQuadLayers to avoid trying to acquire a new swapchain image on static ones.

diff --git a/Engine/Plugins/Runtime/OpenXR/Source/OpenXRHMD/Private/OpenXRHMD.cpp b/Engine/Plugins/Runtime/OpenXR/Source/OpenXRHMD/Private/OpenXRHMD.cpp
--- a/Engine/Plugins/Runtime/OpenXR/Source/OpenXRHMD/Private/OpenXRHMD.cpp
+++ b/Engine/Plugins/Runtime/OpenXR/Source/OpenXRHMD/Private/OpenXRHMD.cpp
@@ -2465,13 +2465,13 @@ void FOpenXRHMD::OnBeginRendering_RenderThread(FRHICommandListImmediate& RHICmdL
 
 		SCOPED_NAMED_EVENT(EnqueueFrame, FColor::Red);
 
-		// Reset the update flag on all layers
-		ForEachLayer([&](uint32 /* unused */, FOpenXRLayer& Layer)
+		// Reset the update flag on native quad layers
+		for (FOpenXRLayer& NativeQuadLayer : NativeQuadLayers)
 		{
-			const bool bUpdateTexture = Layer.Desc.Flags & IStereoLayers::LAYER_FLAG_TEX_CONTINUOUS_UPDATE;
-			Layer.RightEye.bUpdateTexture = bUpdateTexture;
-			Layer.LeftEye.bUpdateTexture = bUpdateTexture;
-		}, false);
+			const bool bUpdateTexture = NativeQuadLayer.Desc.Flags & IStereoLayers::LAYER_FLAG_TEX_CONTINUOUS_UPDATE;
+			NativeQuadLayer.RightEye.bUpdateTexture = bUpdateTexture;
+			NativeQuadLayer.LeftEye.bUpdateTexture = bUpdateTexture;
+		}
 
 		FXRSwapChainPtr ColorSwapchain = PipelinedLayerStateRendering.ColorSwapchain;
 		FXRSwapChainPtr DepthSwapchain = PipelinedLayerStateRendering.DepthSwapchain;
@@ -3382,6 +3382,30 @@ void FOpenXRHMD::DrawVisibleAreaMesh(class FRHICommandList& RHICmdList, int32 Vi
 	}
 }
 
+void FOpenXRHMD::UpdateLayer(FOpenXRLayer& ManagerLayer, uint32 LayerId, bool bIsValid)
+{
+	for (FOpenXRLayer& NativeLayer : NativeQuadLayers)
+	{
+		if (NativeLayer.GetLayerId() == LayerId)
+		{
+			const bool bStaticSwapchain = !(ManagerLayer.Desc.Flags & IStereoLayers::LAYER_FLAG_TEX_CONTINUOUS_UPDATE);
+			NativeLayer.RightEye.bUpdateTexture = ManagerLayer.RightEye.bUpdateTexture;
+			NativeLayer.LeftEye.bUpdateTexture = ManagerLayer.LeftEye.bUpdateTexture;
+			if (bStaticSwapchain)
+			{
+				if (NativeLayer.RightEye.bUpdateTexture) 
+				{
+					NativeLayer.RightEye.Swapchain.Reset();
+				}
+				if (NativeLayer.LeftEye.bUpdateTexture)
+				{
+					NativeLayer.LeftEye.Swapchain.Reset();
+				}
+			}
+		}
+	}
+}
+
 //---------------------------------------------------
 // OpenXR Action Space Implementation
 //---------------------------------------------------
diff --git a/Engine/Plugins/Runtime/OpenXR/Source/OpenXRHMD/Private/OpenXRHMD.h b/Engine/Plugins/Runtime/OpenXR/Source/OpenXRHMD/Private/OpenXRHMD.h
--- a/Engine/Plugins/Runtime/OpenXR/Source/OpenXRHMD/Private/OpenXRHMD.h
+++ b/Engine/Plugins/Runtime/OpenXR/Source/OpenXRHMD/Private/OpenXRHMD.h
@@ -248,6 +248,9 @@ protected:
 	void SetupFrameQuadLayers_RenderThread(FRHICommandListImmediate& RHICmdList);
 	void DrawEmulatedQuadLayers_RenderThread(FRDGBuilder& GraphBuilder, const FSceneView& InView);
 
+	/** TStereoLayerManager<FOpenXRLayer> */
+	void UpdateLayer(FOpenXRLayer& ManagerLayer, uint32 LayerId, bool bIsValid) override;
+
 public:
 	/** IXRTrackingSystem interface */
 	virtual bool DoesSupportLateProjectionUpdate() const override { return true; }
