OpenXR: pipeline background layer visibility state

Before, we'd just check the background layer visibility from the layer manager across all threads.
However, it's only safe to access on the game thread. When we made the background layer invisible, we
could accidentally disable submission of a frame in the RHI thread.

Now we pipeline background layer visibility, _and_ whether we should submit the layer. If we are using layer emulation, we still want to submit the background layer, since the emulation will composite directly into that layer.

diff --git a/Engine/Plugins/Runtime/OpenXR/Source/OpenXRHMD/Private/OpenXRHMD.cpp b/Engine/Plugins/Runtime/OpenXR/Source/OpenXRHMD/Private/OpenXRHMD.cpp
--- a/Engine/Plugins/Runtime/OpenXR/Source/OpenXRHMD/Private/OpenXRHMD.cpp
+++ b/Engine/Plugins/Runtime/OpenXR/Source/OpenXRHMD/Private/OpenXRHMD.cpp
@@ -2403,7 +2403,7 @@ void FOpenXRHMD::DrawEmulatedQuadLayers_RenderThread(FRDGBuilder& GraphBuilder,
 		RHICmdList.BeginRenderPass(RPInfo, TEXT("EmulatedStereoLayerRender"));
 		RHICmdList.SetViewport((float)RenderParams.Viewport.Min.X, (float)RenderParams.Viewport.Min.Y, 0.0f, (float)RenderParams.Viewport.Max.X, (float)RenderParams.Viewport.Max.Y, 1.0f);
 
-		if (bSplashIsShown || !IsBackgroundLayerVisible())
+		if (bSplashIsShown || !PipelinedLayerStateRendering.bBackgroundLayerVisible)
 		{
 			DrawClearQuad(RHICmdList, FLinearColor::Black);
 		}
@@ -2554,12 +2554,16 @@ void FOpenXRHMD::OnBeginRendering_GameThread()
 	bShouldWait = true;
 
 	ENQUEUE_RENDER_COMMAND(TransferFrameStateToRenderingThread)(
-		[this, GameFrameState = PipelinedFrameStateGame](FRHICommandListImmediate& RHICmdList) mutable
+		[this, GameFrameState = PipelinedFrameStateGame, bBackgroundLayerVisible = IsBackgroundLayerVisible()](FRHICommandListImmediate& RHICmdList) mutable
 		{
 			UE_CLOG(PipelinedFrameStateRendering.FrameState.predictedDisplayTime >= GameFrameState.FrameState.predictedDisplayTime,
 				LogHMD, VeryVerbose, TEXT("Predicted display time went backwards from %lld to %lld"), PipelinedFrameStateRendering.FrameState.predictedDisplayTime, GameFrameState.FrameState.predictedDisplayTime);
 
 			PipelinedFrameStateRendering = GameFrameState;
+
+			// If we are emulating layers, we still need to submit background layer since we composite into it
+			PipelinedLayerStateRendering.bBackgroundLayerVisible = bBackgroundLayerVisible;
+			PipelinedLayerStateRendering.bSubmitBackgroundLayer = bBackgroundLayerVisible || !bNativeWorldQuadLayerSupport;
 		});
 }
 
@@ -2894,7 +2898,7 @@ void FOpenXRHMD::OnFinishRendering_RHIThread()
 	{
 		TArray<const XrCompositionLayerBaseHeader*> Headers;
 		XrCompositionLayerProjection Layer = {};
-		if (IsBackgroundLayerVisible())
+		if (PipelinedLayerStateRHI.bSubmitBackgroundLayer)
 		{
 			Layer.type = XR_TYPE_COMPOSITION_LAYER_PROJECTION;
 			Layer.next = nullptr;
diff --git a/Engine/Plugins/Runtime/OpenXR/Source/OpenXRHMD/Private/OpenXRHMD.h b/Engine/Plugins/Runtime/OpenXR/Source/OpenXRHMD/Private/OpenXRHMD.h
--- a/Engine/Plugins/Runtime/OpenXR/Source/OpenXRHMD/Private/OpenXRHMD.h
+++ b/Engine/Plugins/Runtime/OpenXR/Source/OpenXRHMD/Private/OpenXRHMD.h
@@ -103,6 +103,9 @@ public:
 		FXRSwapChainPtr ColorSwapchain;
 		FXRSwapChainPtr DepthSwapchain;
 		TArray<FXRSwapChainPtr> QuadSwapchains;
+
+		bool bBackgroundLayerVisible = true;
+		bool bSubmitBackgroundLayer = true;
 	};
 
 	class FVulkanExtensions : public IHeadMountedDisplayVulkanExtensions
