Add Blueprint node to set Environment Blend Mode in OpenXR plugin.

diff --git a/Engine/Plugins/Runtime/OpenXR/Source/OpenXRHMD/Private/OpenXRBlueprintFunctionLibrary.cpp b/Engine/Plugins/Runtime/OpenXR/Source/OpenXRHMD/Private/OpenXRBlueprintFunctionLibrary.cpp
new file mode 100644
--- /dev/null
+++ b/Engine/Plugins/Runtime/OpenXR/Source/OpenXRHMD/Private/OpenXRBlueprintFunctionLibrary.cpp
@@ -0,0 +1,33 @@
+// Copyright Epic Games, Inc. All Rights Reserved.
+
+#include <OpenXRBlueprintFunctionLibrary.h>
+#include <OpenXRHMD.h>
+
+#include UE_INLINE_GENERATED_CPP_BY_NAME(OpenXRBlueprintFunctionLibrary)
+
+DEFINE_LOG_CATEGORY_STATIC(LogUOpenXRBlueprintFunctionLibrary, Log, All);
+
+UOpenXRBlueprintFunctionLibrary::UOpenXRBlueprintFunctionLibrary(FObjectInitializer const& ObjectInitializer)
+	: Super(ObjectInitializer)
+{
+}
+
+void UOpenXRBlueprintFunctionLibrary::SetEnvironmentBlendMode(int32 NewBlendMode) 
+{
+	FOpenXRHMD* OpenXRHMD = GetOpenXRHMD();
+	if (OpenXRHMD) 
+	{
+		OpenXRHMD->SetEnvironmentBlendMode((XrEnvironmentBlendMode) NewBlendMode);
+	}
+}
+
+FOpenXRHMD* UOpenXRBlueprintFunctionLibrary::GetOpenXRHMD()
+{
+	static FName SystemName(TEXT("OpenXR"));
+	if (GEngine->XRSystem.IsValid() && (GEngine->XRSystem->GetSystemName() == SystemName))
+	{
+		return static_cast<FOpenXRHMD*>(GEngine->XRSystem.Get());
+	}
+
+	return nullptr;
+}
diff --git a/Engine/Plugins/Runtime/OpenXR/Source/OpenXRHMD/Private/OpenXRHMD.cpp b/Engine/Plugins/Runtime/OpenXR/Source/OpenXRHMD/Private/OpenXRHMD.cpp
--- a/Engine/Plugins/Runtime/OpenXR/Source/OpenXRHMD/Private/OpenXRHMD.cpp
+++ b/Engine/Plugins/Runtime/OpenXR/Source/OpenXRHMD/Private/OpenXRHMD.cpp
@@ -63,6 +63,18 @@ namespace {
 	static TSet<XrEnvironmentBlendMode> SupportedBlendModes{ XR_ENVIRONMENT_BLEND_MODE_ALPHA_BLEND, XR_ENVIRONMENT_BLEND_MODE_ADDITIVE, XR_ENVIRONMENT_BLEND_MODE_OPAQUE };
 	static TSet<XrViewConfigurationType> SupportedViewConfigurations{ XR_VIEW_CONFIGURATION_TYPE_PRIMARY_STEREO, XR_VIEW_CONFIGURATION_TYPE_PRIMARY_QUAD_VARJO };
 
+	/** Helper function for checking whether a blend mode is supported. */
+	bool IsEnvironmentBlendModeSupported(XrEnvironmentBlendMode BlendMode)
+	{
+		if (SupportedBlendModes.Contains(BlendMode) &&
+			// On mobile platforms the alpha channel can contain depth information, so we can't use alpha blend.
+			(BlendMode != XR_ENVIRONMENT_BLEND_MODE_ALPHA_BLEND || !IsMobilePlatform(GMaxRHIShaderPlatform)))
+		{
+			return true;
+		}
+		return false;
+	}
+
 	/** Helper function for acquiring the appropriate FSceneViewport */
 	FSceneViewport* FindSceneViewport()
 	{
@@ -1268,22 +1280,13 @@ FOpenXRHMD::FOpenXRHMD(const FAutoRegister& AutoRegister, XrInstance InInstance,
 
 	// Enumerate environment blend modes and select the best one.
 	{
-		uint32 BlendModeCount;
-		TArray<XrEnvironmentBlendMode> BlendModes;
-		XR_ENSURE(xrEnumerateEnvironmentBlendModes(Instance, System, SelectedViewConfigurationType, 0, &BlendModeCount, nullptr));
-		// Fill the initial array with valid enum types (this will fail in the validation layer otherwise).
-		for (auto& TypeIter : BlendModes)
-			TypeIter = XR_ENVIRONMENT_BLEND_MODE_OPAQUE;
-		BlendModes.SetNum(BlendModeCount);
-		XR_ENSURE(xrEnumerateEnvironmentBlendModes(Instance, System, SelectedViewConfigurationType, BlendModeCount, &BlendModeCount, BlendModes.GetData()));
+		TArray<XrEnvironmentBlendMode> BlendModes = RetrieveEnvironmentBlendModes();
 
 		// Select the first blend mode returned by the runtime that is supported.
 		// This is the environment blend mode preferred by the runtime.
 		for (XrEnvironmentBlendMode BlendMode : BlendModes)
 		{
-			if (SupportedBlendModes.Contains(BlendMode) &&
-				// On mobile platforms the alpha channel can contain depth information, so we can't use alpha blend.
-				(BlendMode != XR_ENVIRONMENT_BLEND_MODE_ALPHA_BLEND || !IsMobilePlatform(GMaxRHIShaderPlatform)))
+			if (IsEnvironmentBlendModeSupported(BlendMode))
 			{
 				SelectedEnvironmentBlendMode = BlendMode;
 				break;
@@ -1341,6 +1344,17 @@ bool FOpenXRHMD::ReconfigureForShaderPlatform(EShaderPlatform NewShaderPlatform)
 	return true;
 }
 
+TArray<XrEnvironmentBlendMode> FOpenXRHMD::RetrieveEnvironmentBlendModes() const
+{
+	TArray<XrEnvironmentBlendMode> BlendModes;
+	uint32 BlendModeCount;
+	XR_ENSURE(xrEnumerateEnvironmentBlendModes(Instance, System, SelectedViewConfigurationType, 0, &BlendModeCount, nullptr));
+	// Fill the initial array with valid enum types (this will fail in the validation layer otherwise).
+	BlendModes.Init(XR_ENVIRONMENT_BLEND_MODE_OPAQUE, BlendModeCount);
+	XR_ENSURE(xrEnumerateEnvironmentBlendModes(Instance, System, SelectedViewConfigurationType, BlendModeCount, &BlendModeCount, BlendModes.GetData()));
+	return BlendModes;
+}
+
 const FOpenXRHMD::FPipelinedFrameState& FOpenXRHMD::GetPipelinedFrameStateForThread() const
 {
 	// Relying on implicit selection of the RHI struct is hazardous since the RHI thread isn't always present
@@ -1992,6 +2006,32 @@ bool FOpenXRHMD::IsFocused() const
 	return CurrentSessionState == XR_SESSION_STATE_FOCUSED;
 }
 
+void FOpenXRHMD::SetEnvironmentBlendMode(XrEnvironmentBlendMode NewBlendMode) 
+{
+	if (NewBlendMode == XR_ENVIRONMENT_BLEND_MODE_MAX_ENUM) 
+	{
+		UE_LOG(LogHMD, Error, TEXT("Environment Blend Mode can't be set to XR_ENVIRONMENT_BLEND_MODE_MAX_ENUM."));
+		return;
+	}
+
+	if(!Instance || !System)
+	{
+		return;
+	}
+
+	TArray<XrEnvironmentBlendMode> BlendModes = RetrieveEnvironmentBlendModes();
+
+	if (BlendModes.Contains(NewBlendMode) && IsEnvironmentBlendModeSupported(NewBlendMode))
+	{
+		SelectedEnvironmentBlendMode = NewBlendMode;
+		UE_LOG(LogHMD, Log, TEXT("Environment Blend Mode set to: %d."), SelectedEnvironmentBlendMode);
+	}
+	else
+	{
+		UE_LOG(LogHMD, Error, TEXT("Environment Blend Mode %d is not supported. Environment Blend Mode remains %d."), NewBlendMode, SelectedEnvironmentBlendMode);
+	}
+}
+
 bool FOpenXRHMD::StartSession()
 {
 	// If the session is not yet ready, we'll call into this function again when it is
diff --git a/Engine/Plugins/Runtime/OpenXR/Source/OpenXRHMD/Private/OpenXRHMD.h b/Engine/Plugins/Runtime/OpenXR/Source/OpenXRHMD/Private/OpenXRHMD.h
--- a/Engine/Plugins/Runtime/OpenXR/Source/OpenXRHMD/Private/OpenXRHMD.h
+++ b/Engine/Plugins/Runtime/OpenXR/Source/OpenXRHMD/Private/OpenXRHMD.h
@@ -370,8 +370,11 @@ public:
 	OPENXRHMD_API XrTime GetDisplayTime() const;
 	OPENXRHMD_API XrSpace GetTrackingSpace() const;
 	OPENXRHMD_API TArray<IOpenXRExtensionPlugin*>& GetExtensionPlugins() { return ExtensionPlugins; }
+	OPENXRHMD_API void SetEnvironmentBlendMode(XrEnvironmentBlendMode NewBlendMode);
 
 private:
+	TArray<XrEnvironmentBlendMode> RetrieveEnvironmentBlendModes() const;
+
 	bool					bStereoEnabled;
 	TAtomic<bool>			bIsRunning;
 	TAtomic<bool>			bIsReady;
diff --git a/Engine/Plugins/Runtime/OpenXR/Source/OpenXRHMD/Public/OpenXRBlueprintFunctionLibrary.h b/Engine/Plugins/Runtime/OpenXR/Source/OpenXRHMD/Public/OpenXRBlueprintFunctionLibrary.h
new file mode 100644
--- /dev/null
+++ b/Engine/Plugins/Runtime/OpenXR/Source/OpenXRHMD/Public/OpenXRBlueprintFunctionLibrary.h
@@ -0,0 +1,26 @@
+// Copyright Epic Games, Inc. All Rights Reserved.
+
+#pragma once
+
+#include "Kismet/BlueprintFunctionLibrary.h"
+
+#include "OpenXRBlueprintFunctionLibrary.generated.h"
+
+class FOpenXRHMD;
+
+UCLASS()
+class OPENXRHMD_API UOpenXRBlueprintFunctionLibrary : public UBlueprintFunctionLibrary
+{
+	GENERATED_UCLASS_BODY()
+	/**
+	 * Allows to change the environment blend mode at runtime.
+	 * 
+	 * @param NewBlendMode		The new blend mode to be set, if supported.
+	 */
+	UFUNCTION(BlueprintCallable, Category="OpenXR|Experimental")
+	static void SetEnvironmentBlendMode(int32 NewBlendMode);
+
+private:
+	static FOpenXRHMD* GetOpenXRHMD();
+};
+
