Add a sync point with the RHI thread to avoid a race condition on the XR swapchain

diff --git a/Engine/Plugins/Runtime/OpenXR/Source/OpenXRHMD/Private/OpenXRHMD.cpp b/Engine/Plugins/Runtime/OpenXR/Source/OpenXRHMD/Private/OpenXRHMD.cpp
--- a/Engine/Plugins/Runtime/OpenXR/Source/OpenXRHMD/Private/OpenXRHMD.cpp
+++ b/Engine/Plugins/Runtime/OpenXR/Source/OpenXRHMD/Private/OpenXRHMD.cpp
@@ -2522,6 +2522,10 @@ void FOpenXRHMD::OnBeginRendering_RenderThread(FRHICommandListImmediate& RHICmdL
 		{
 			OnBeginRendering_RHIThread(FrameState, ColorSwapchain, DepthSwapchain);
 		});
+
+		// We need to sync with the RHI thread to ensure we've acquired the next swapchain image.
+		// TODO: The acquire needs to be moved to the Render thread as soon as it's allowed by the spec.
+		RHICmdList.ImmediateFlush(EImmediateFlushType::FlushRHIThread);
 	}
 
 	// Snapshot new poses for late update.
