Pico-VR Development Commits PR #9412 for Nico: Updating assets for controllers #jira UE-160262 #rb Rob.Srinivasiah

diff --git a/Engine/Plugins/Runtime/OpenXR/Source/OpenXRHMD/Private/OpenXRAssetManager.cpp b/Engine/Plugins/Runtime/OpenXR/Source/OpenXRHMD/Private/OpenXRAssetManager.cpp
--- a/Engine/Plugins/Runtime/OpenXR/Source/OpenXRHMD/Private/OpenXRAssetManager.cpp
+++ b/Engine/Plugins/Runtime/OpenXR/Source/OpenXRHMD/Private/OpenXRAssetManager.cpp
@@ -32,9 +32,8 @@ FSoftObjectPath FOpenXRAssetDirectory::OculusTouchV2Left          = FString(TEXT
 FSoftObjectPath FOpenXRAssetDirectory::OculusTouchV2Right         = FString(TEXT("/OpenXR/Devices/OculusTouch_v2/Right/right_OculusTouch_v2Controller.right_OculusTouch_v2Controller"));
 FSoftObjectPath FOpenXRAssetDirectory::OculusTouchV3Left          = FString(TEXT("/OpenXR/Devices/OculusTouch_v3/Left/left_OculusTouch_v3Controller.left_OculusTouch_v3Controller"));
 FSoftObjectPath FOpenXRAssetDirectory::OculusTouchV3Right         = FString(TEXT("/OpenXR/Devices/OculusTouch_v3/Right/right_OculusTouch_v3Controller.right_OculusTouch_v3Controller"));
-FSoftObjectPath FOpenXRAssetDirectory::PicoG2                     = FString(TEXT("/OpenXR/Devices/PicoG2/PicoG2Controller.PicoG2Controller"));
-FSoftObjectPath FOpenXRAssetDirectory::PicoNeo2Left               = FString(TEXT("/OpenXR/Devices/PicoNeo2/Left/left_PicoNeo2Controller.left_PicoNeo2Controller"));
-FSoftObjectPath FOpenXRAssetDirectory::PicoNeo2Right              = FString(TEXT("/OpenXR/Devices/PicoNeo2/Right/right_PicoNeo2Controller.right_PicoNeo2Controller"));
+FSoftObjectPath FOpenXRAssetDirectory::PicoNeo3Left				  = FString(TEXT("/OpenXR/Devices/PicoNeo3/Left/left_PicoNeo3Controller.left_PicoNeo3Controller"));
+FSoftObjectPath FOpenXRAssetDirectory::PicoNeo3Right              = FString(TEXT("/OpenXR/Devices/PicoNeo3/Right/right_PicoNeo3Controller.right_PicoNeo3Controller"));
 FSoftObjectPath FOpenXRAssetDirectory::SamsungGearVR              = FString(TEXT("/OpenXR/Devices/SamsungGearVR/SamsungGearVRController.SamsungGearVRController"));
 FSoftObjectPath FOpenXRAssetDirectory::SamsungOdysseyLeft         = FString(TEXT("/OpenXR/Devices/SamsungOdyssey/Left/left_SamsungOdysseyController.left_SamsungOdysseyController"));
 FSoftObjectPath FOpenXRAssetDirectory::SamsungOdysseyRight        = FString(TEXT("/OpenXR/Devices/SamsungOdyssey/Right/right_SamsungOdysseyController.right_SamsungOdysseyController"));
diff --git a/Engine/Plugins/Runtime/OpenXR/Source/OpenXRHMD/Public/OpenXRAssetDirectory.h b/Engine/Plugins/Runtime/OpenXR/Source/OpenXRHMD/Public/OpenXRAssetDirectory.h
--- a/Engine/Plugins/Runtime/OpenXR/Source/OpenXRHMD/Public/OpenXRAssetDirectory.h
+++ b/Engine/Plugins/Runtime/OpenXR/Source/OpenXRHMD/Public/OpenXRAssetDirectory.h
@@ -29,9 +29,8 @@ public:
 	OPENXRHMD_API static FSoftObjectPath OculusTouchV2Right;
 	OPENXRHMD_API static FSoftObjectPath OculusTouchV3Left;
 	OPENXRHMD_API static FSoftObjectPath OculusTouchV3Right;
-	OPENXRHMD_API static FSoftObjectPath PicoG2;
-	OPENXRHMD_API static FSoftObjectPath PicoNeo2Left;
-	OPENXRHMD_API static FSoftObjectPath PicoNeo2Right;
+	OPENXRHMD_API static FSoftObjectPath PicoNeo3Left;
+	OPENXRHMD_API static FSoftObjectPath PicoNeo3Right;
 	OPENXRHMD_API static FSoftObjectPath SamsungGearVR;
 	OPENXRHMD_API static FSoftObjectPath SamsungOdysseyLeft;
 	OPENXRHMD_API static FSoftObjectPath SamsungOdysseyRight;
