PR #9899: Destroy xrInstance if HMD is not connected. (Contributed by [REDACTED_PERSON])

diff --git a/Engine/Plugins/Runtime/OpenXR/Source/OpenXRHMD/Private/OpenXRHMDModule.cpp b/Engine/Plugins/Runtime/OpenXR/Source/OpenXRHMD/Private/OpenXRHMDModule.cpp
--- a/Engine/Plugins/Runtime/OpenXR/Source/OpenXRHMD/Private/OpenXRHMDModule.cpp
+++ b/Engine/Plugins/Runtime/OpenXR/Source/OpenXRHMD/Private/OpenXRHMDModule.cpp
@@ -74,7 +74,7 @@ void FOpenXRHMDModule::ShutdownModule()
 {
 	if (Instance)
 	{
-		XR_ENSURE(xrDestroyInstance(Instance));
+		DestroyInstance();
 	}
 
 	if (LoaderHandle)
@@ -745,6 +745,7 @@ bool FOpenXRHMDModule::InitSystem()
 	XrResult Result = xrGetSystem(Instance, &SystemInfo, &System);
 	if (XR_FAILED(Result))
 	{
+		DestroyInstance();
 		UE_LOG(LogHMD, Log, TEXT("Failed to get an OpenXR system, result is %s. Please check that your runtime supports VR headsets."), OpenXRResultToString(Result));
 		return false;
 	}
@@ -757,6 +758,12 @@ bool FOpenXRHMDModule::InitSystem()
 	return true;
 }
 
+void FOpenXRHMDModule::DestroyInstance()
+{
+	XR_ENSURE(xrDestroyInstance(Instance));
+	Instance = XR_NULL_HANDLE;
+}
+
 FName FOpenXRHMDModule::ResolvePathToName(XrPath Path)
 {
 	{
diff --git a/Engine/Plugins/Runtime/OpenXR/Source/OpenXRHMD/Private/OpenXRHMDModule.h b/Engine/Plugins/Runtime/OpenXR/Source/OpenXRHMD/Private/OpenXRHMDModule.h
--- a/Engine/Plugins/Runtime/OpenXR/Source/OpenXRHMD/Private/OpenXRHMDModule.h
+++ b/Engine/Plugins/Runtime/OpenXR/Source/OpenXRHMD/Private/OpenXRHMDModule.h
@@ -61,6 +61,7 @@ private:
 	TSortedMap<XrPath, FName> PathToName;
 	TSortedMap<FName, XrPath, FDefaultAllocator, FNameFastLess> NameToPath;
 
+	void DestroyInstance();
 	bool EnumerateExtensions();
 	bool EnumerateLayers();
 	bool InitRenderBridge();
