Workaround issues on SteamVR with long haptics durations.

diff --git a/Engine/Plugins/Runtime/OpenXR/Source/OpenXRInput/Private/OpenXRInput.cpp b/Engine/Plugins/Runtime/OpenXR/Source/OpenXRInput/Private/OpenXRInput.cpp
--- a/Engine/Plugins/Runtime/OpenXR/Source/OpenXRInput/Private/OpenXRInput.cpp
+++ b/Engine/Plugins/Runtime/OpenXR/Source/OpenXRInput/Private/OpenXRInput.cpp
@@ -1271,7 +1271,7 @@ void FOpenXRInputPlugin::FOpenXRInput::SetHapticFeedbackValues(int32 ControllerI
 	XrHapticVibration HapticValue;
 	HapticValue.type = XR_TYPE_HAPTIC_VIBRATION;
 	HapticValue.next = nullptr;
-	HapticValue.duration = MaxFeedbackDuration;
+	HapticValue.duration = XrDuration(CurrentDeltaTime * 1e9);
 	HapticValue.frequency = Values.Frequency;
 	HapticValue.amplitude = Values.Amplitude;
 
diff --git a/Engine/Plugins/Runtime/OpenXR/Source/OpenXRInput/Private/OpenXRInput.h b/Engine/Plugins/Runtime/OpenXR/Source/OpenXRInput/Private/OpenXRInput.h
--- a/Engine/Plugins/Runtime/OpenXR/Source/OpenXRInput/Private/OpenXRInput.h
+++ b/Engine/Plugins/Runtime/OpenXR/Source/OpenXRInput/Private/OpenXRInput.h
@@ -108,7 +108,7 @@ public:
 		virtual void OnDestroySession() override;
 
 		// IInputDevice overrides
-		virtual void Tick(float DeltaTime) override {};
+		virtual void Tick(float DeltaTime) override { CurrentDeltaTime = DeltaTime; };
 		virtual void SendControllerEvents() override;
 		virtual void SetMessageHandler(const TSharedRef< FGenericApplicationMessageHandler >& InMessageHandler) override;
 		virtual bool Exec(UWorld* InWorld, const TCHAR* Cmd, FOutputDevice& Ar) override;
@@ -134,8 +134,6 @@ public:
 		virtual float GetHapticAmplitudeScale() const override;
 
 	private:
-		static const XrDuration MaxFeedbackDuration = 2500000000; // 2.5 seconds
-
 		FOpenXRHMD* OpenXRHMD;
 		XrInstance Instance;
 
@@ -156,6 +154,11 @@ public:
 		bool bActionsAttached;
 		bool bDirectionalBindingSupported;
 
+		/**
+		* Buffer for current delta time to get an accurate approximation of how long to play haptics for
+		*/
+		float CurrentDeltaTime = 0.0f;
+
 		bool BuildActions(XrSession Session);
 		void SyncActions(XrSession Session);
 		void BuildLegacyActions(TMap<FString, FInteractionProfile>& Profiles);
