Pass format list to FOpenXRHMD::CreateSwapchain only when a non-sRGB format is requested.
Reworked from PR 9924.

diff --git a/Engine/Plugins/Runtime/OpenXR/Source/OpenXRHMD/Private/OpenXRHMD_Swapchain.cpp b/Engine/Plugins/Runtime/OpenXR/Source/OpenXRHMD/Private/OpenXRHMD_Swapchain.cpp
--- a/Engine/Plugins/Runtime/OpenXR/Source/OpenXRHMD/Private/OpenXRHMD_Swapchain.cpp
+++ b/Engine/Plugins/Runtime/OpenXR/Source/OpenXRHMD/Private/OpenXRHMD_Swapchain.cpp
@@ -409,7 +409,13 @@ FXRSwapChainPtr CreateSwapchain_Vulkan(XrSession InSession, uint8 Format, uint8&
 	FormatListInfo.viewFormatCount = ViewFormatList.Num();
 	FormatListInfo.viewFormats = ViewFormatList.GetData();
 
-	XrSwapchain Swapchain = FOpenXRSwapchain::CreateSwapchain(InSession, ViewFormatList[0], SizeX, SizeY, ArraySize, NumMips, NumSamples, CreateFlags, &FormatListInfo);
+	// OpenXR wants to create sRGB swapchains. When the swapchain that is being created does not conform to this,
+	// we need to tell OpenXR what formats the swapchain will be accessed in.
+	// The XR_SWAPCHAIN_USAGE_MUTABLE_FORMAT_BIT will also be set in FOpenXRSwapchain::CreateSwapchain.
+	// We don't need to specify additional formats if the swapchain is being created as sRGB.
+	void* Next = !(CreateFlags & TexCreate_SRGB) ? &FormatListInfo : nullptr;
+
+	XrSwapchain Swapchain = FOpenXRSwapchain::CreateSwapchain(InSession, ViewFormatList[0], SizeX, SizeY, ArraySize, NumMips, NumSamples, CreateFlags, Next);
 	if (!Swapchain)
 	{
 		return nullptr;
