Allow the OpenXR plugin to initialize without a headset connected.

Adds the ability to connect a headset while the editor is already running.

diff --git a/Engine/Plugins/Runtime/OpenXR/Source/OpenXRHMD/Private/OpenXRHMD.cpp b/Engine/Plugins/Runtime/OpenXR/Source/OpenXRHMD/Private/OpenXRHMD.cpp
--- a/Engine/Plugins/Runtime/OpenXR/Source/OpenXRHMD/Private/OpenXRHMD.cpp
+++ b/Engine/Plugins/Runtime/OpenXR/Source/OpenXRHMD/Private/OpenXRHMD.cpp
@@ -6,6 +6,7 @@
 #include "OpenXRHMD_Swapchain.h"
 #include "OpenXRCore.h"
 #include "IOpenXRExtensionPlugin.h"
+#include "IOpenXRHMDModule.h"
 
 #include "Misc/App.h"
 #include "Misc/ConfigCacheIni.h"
@@ -41,8 +42,12 @@
 #if WITH_EDITOR
 #include "Editor.h"
 #include "Editor/EditorEngine.h"
+#include "Misc/MessageDialog.h"
+#include "UnrealEdMisc.h"
 #endif
 
+#define LOCTEXT_NAMESPACE "OpenXR"
+
 #define OPENXR_PAUSED_IDLE_FPS 10
 static const int64 OPENXR_SWAPCHAIN_WAIT_TIMEOUT = 100000000ll;		// 100ms in nanoseconds.
 
@@ -466,6 +471,11 @@
 		XR_VERSION_PATCH(InstanceProperties.runtimeVersion));
 }
 
+bool FOpenXRHMD::IsHMDConnected()
+{
+	return IOpenXRHMDModule::Get().GetSystemId() != XR_NULL_SYSTEM_ID;
+}
+
 bool FOpenXRHMD::IsHMDEnabled() const
 {
 	return true;
@@ -1176,7 +1186,7 @@
 	return true;
 }
 
-FOpenXRHMD::FOpenXRHMD(const FAutoRegister& AutoRegister, XrInstance InInstance, XrSystemId InSystem, TRefCountPtr<FOpenXRRenderBridge>& InRenderBridge, TArray<const char*> InEnabledExtensions, TArray<IOpenXRExtensionPlugin*> InExtensionPlugins, IARSystemSupport* ARSystemSupport)
+FOpenXRHMD::FOpenXRHMD(const FAutoRegister& AutoRegister, XrInstance InInstance, TRefCountPtr<FOpenXRRenderBridge>& InRenderBridge, TArray<const char*> InEnabledExtensions, TArray<IOpenXRExtensionPlugin*> InExtensionPlugins, IARSystemSupport* ARSystemSupport)
 	: FHeadMountedDisplayBase(ARSystemSupport)
 	, FHMDSceneViewExtension(AutoRegister)
 	, FOpenXRAssetManager(InInstance, this)
@@ -1197,7 +1207,7 @@
 	, InputModule(nullptr)
 	, ExtensionPlugins(std::move(InExtensionPlugins))
 	, Instance(InInstance)
-	, System(InSystem)
+	, System(XR_NULL_SYSTEM_ID)
 	, Session(XR_NULL_HANDLE)
 	, LocalSpace(XR_NULL_HANDLE)
 	, StageSpace(XR_NULL_HANDLE)
@@ -1223,83 +1233,11 @@
 	bHiddenAreaMaskSupported = IsExtensionEnabled(XR_KHR_VISIBILITY_MASK_EXTENSION_NAME) &&
 		!FCStringAnsi::Strstr(InstanceProperties.runtimeName, "Oculus");
 	bViewConfigurationFovSupported = IsExtensionEnabled(XR_EPIC_VIEW_CONFIGURATION_FOV_EXTENSION_NAME);
-
-	// Retrieve system properties and check for hand tracking support
-	XrSystemHandTrackingPropertiesEXT HandTrackingSystemProperties = { XR_TYPE_SYSTEM_HAND_TRACKING_PROPERTIES_EXT };
-	SystemProperties = XrSystemProperties{ XR_TYPE_SYSTEM_PROPERTIES, &HandTrackingSystemProperties };
-	XR_ENSURE(xrGetSystemProperties(Instance, System, &SystemProperties));
-	bSupportsHandTracking = HandTrackingSystemProperties.supportsHandTracking == XR_TRUE;
-	SystemProperties.next = nullptr;
-
-	// Some runtimes aren't compliant with their number of layers supported.
-	// We support a fallback by emulating non-facelocked layers
-	bNativeWorldQuadLayerSupport = SystemProperties.graphicsProperties.maxLayerCount >= XR_MIN_COMPOSITION_LAYERS_SUPPORTED;
-
+	bSupportsHandTracking = IsExtensionEnabled(XR_EXT_HAND_TRACKING_EXTENSION_NAME);
 	bSpaceAccellerationSupported = IsExtensionEnabled(XR_EPIC_SPACE_ACCELERATION_NAME);
 
 	ReconfigureForShaderPlatform(GMaxRHIShaderPlatform);
 
-	// Enumerate the viewport configurations
-	uint32 ConfigurationCount;
-	TArray<XrViewConfigurationType> ViewConfigTypes;
-	XR_ENSURE(xrEnumerateViewConfigurations(Instance, System, 0, &ConfigurationCount, nullptr));
-	ViewConfigTypes.SetNum(ConfigurationCount);
-	// Fill the initial array with valid enum types (this will fail in the validation layer otherwise).
-	for (auto & TypeIter : ViewConfigTypes)
-		TypeIter = XR_VIEW_CONFIGURATION_TYPE_PRIMARY_MONO;
-	XR_ENSURE(xrEnumerateViewConfigurations(Instance, System, ConfigurationCount, &ConfigurationCount, ViewConfigTypes.GetData()));
-
-	// Select the first view configuration returned by the runtime that is supported.
-	// This is the view configuration preferred by the runtime.
-	for (XrViewConfigurationType ViewConfigType : ViewConfigTypes)
-	{
-		if (SupportedViewConfigurations.Contains(ViewConfigType))
-		{
-			SelectedViewConfigurationType = ViewConfigType;
-			break;
-		}
-	}
-
-	// If there is no supported view configuration type, use the first option as a last resort.
-	if (!ensure(SelectedViewConfigurationType != XR_VIEW_CONFIGURATION_TYPE_MAX_ENUM))
-	{
-		UE_LOG(LogHMD, Error, TEXT("No compatible view configuration type found, falling back to runtime preferred type."));
-		SelectedViewConfigurationType = ViewConfigTypes[0];
-	}
-
-	// Enumerate the views we will be simulating with.
-	EnumerateViews(PipelinedFrameStateGame);
-
-	for (const XrViewConfigurationView& Config : PipelinedFrameStateGame.ViewConfigs)
-	{
-		const float WidthDensityMax = float(Config.maxImageRectWidth) / Config.recommendedImageRectWidth;
-		const float HeightDensitymax = float(Config.maxImageRectHeight) / Config.recommendedImageRectHeight;
-		const float PerViewPixelDensityMax = FMath::Min(WidthDensityMax, HeightDensitymax);
-		RuntimePixelDensityMax = FMath::Min(RuntimePixelDensityMax, PerViewPixelDensityMax);
-	}
-
-	// Enumerate environment blend modes and select the best one.
-	{
-		TArray<XrEnvironmentBlendMode> BlendModes = RetrieveEnvironmentBlendModes();
-
-		// Select the first blend mode returned by the runtime that is supported.
-		// This is the environment blend mode preferred by the runtime.
-		for (XrEnvironmentBlendMode BlendMode : BlendModes)
-		{
-			if (IsEnvironmentBlendModeSupported(BlendMode))
-			{
-				SelectedEnvironmentBlendMode = BlendMode;
-				break;
-			}
-		}
-
-		// If there is no supported environment blend mode, use the first option as a last resort.
-		if (!ensure(SelectedEnvironmentBlendMode != XR_ENVIRONMENT_BLEND_MODE_MAX_ENUM))
-		{
-			SelectedEnvironmentBlendMode = BlendModes[0];
-		}
-	}
-
 #if PLATFORM_HOLOLENS || PLATFORM_ANDROID
 	bool bStartInVR = false;
 	GConfig->GetBool(TEXT("/Script/EngineSettings.GeneralProjectSettings"), TEXT("bStartInVR"), bStartInVR, GGameIni);
@@ -1319,9 +1257,6 @@
 	XR_ENSURE(xrStringToPath(Instance, "/user/head", &UserHead));
 	ensure(DeviceSpaces.Emplace(XR_NULL_HANDLE, UserHead) == HMDDeviceId);
 
-	// Give the all frame states the same initial values.
-	PipelinedFrameStateRHI = PipelinedFrameStateRendering = PipelinedFrameStateGame;
-
 	for (IOpenXRExtensionPlugin* Module : ExtensionPlugins)
 	{
 		Module->BindExtensionPluginDelegates(*this);
@@ -1669,6 +1604,19 @@
 }
 #endif
 
+#if WITH_EDITOR
+// Show a warning that the editor will require a restart
+void ShowRestartWarning(const FText& Title)
+{
+	if (EAppReturnType::Ok == FMessageDialog::Open(EAppMsgType::OkCancel,
+		LOCTEXT("EditorRestartMsg", "The OpenXR runtime requires switching to a different GPU adapter, this requires an editor restart. Do you wish to restart now (you will be prompted to save any changes)?"),
+		&Title))
+	{
+		FUnrealEdMisc::Get().RestartEditor(false);
+	}
+}
+#endif
+
 bool FOpenXRHMD::OnStereoStartup()
 {
 	FWriteScopeLock Lock(SessionHandleMutex);
@@ -1676,9 +1624,102 @@
 	check(Session == XR_NULL_HANDLE);
 	bIsExitingSessionByxrRequestExitSession = false;  // clear in case we requested exit for a previous session, but it ended in some other way before that happened.
 
+	System = IOpenXRHMDModule::Get().GetSystemId();
+	if (!System)
+	{
+		UE_LOG(LogHMD, Error, TEXT("Failed to get an OpenXR system, please check that you have a VR headset connected."));
+		return false;
+	}
+
+	// Retrieve system properties and check for hand tracking support
+	XrSystemHandTrackingPropertiesEXT HandTrackingSystemProperties = { XR_TYPE_SYSTEM_HAND_TRACKING_PROPERTIES_EXT };
+	SystemProperties = XrSystemProperties{ XR_TYPE_SYSTEM_PROPERTIES, &HandTrackingSystemProperties };
+	XR_ENSURE(xrGetSystemProperties(Instance, System, &SystemProperties));
+	bSupportsHandTracking = HandTrackingSystemProperties.supportsHandTracking == XR_TRUE;
+
+	// Some runtimes aren't compliant with their number of layers supported.
+	// We support a fallback by emulating non-facelocked layers
+	bNativeWorldQuadLayerSupport = SystemProperties.graphicsProperties.maxLayerCount >= XR_MIN_COMPOSITION_LAYERS_SUPPORTED;
+
+	// Enumerate the viewport configurations
+	uint32 ConfigurationCount;
+	TArray<XrViewConfigurationType> ViewConfigTypes;
+	XR_ENSURE(xrEnumerateViewConfigurations(Instance, System, 0, &ConfigurationCount, nullptr));
+	ViewConfigTypes.SetNum(ConfigurationCount);
+	// Fill the initial array with valid enum types (this will fail in the validation layer otherwise).
+	for (auto& TypeIter : ViewConfigTypes)
+		TypeIter = XR_VIEW_CONFIGURATION_TYPE_PRIMARY_MONO;
+	XR_ENSURE(xrEnumerateViewConfigurations(Instance, System, ConfigurationCount, &ConfigurationCount, ViewConfigTypes.GetData()));
+
+	// Select the first view configuration returned by the runtime that is supported.
+	// This is the view configuration preferred by the runtime.
+	for (XrViewConfigurationType ViewConfigType : ViewConfigTypes)
+	{
+		if (SupportedViewConfigurations.Contains(ViewConfigType))
+		{
+			SelectedViewConfigurationType = ViewConfigType;
+			break;
+		}
+	}
+
+	// If there is no supported view configuration type, use the first option as a last resort.
+	if (!ensure(SelectedViewConfigurationType != XR_VIEW_CONFIGURATION_TYPE_MAX_ENUM))
+	{
+		UE_LOG(LogHMD, Error, TEXT("No compatible view configuration type found, falling back to runtime preferred type."));
+		SelectedViewConfigurationType = ViewConfigTypes[0];
+	}
+
+	// Enumerate the views we will be simulating with.
+	EnumerateViews(PipelinedFrameStateGame);
+
+	for (const XrViewConfigurationView& Config : PipelinedFrameStateGame.ViewConfigs)
+	{
+		const float WidthDensityMax = float(Config.maxImageRectWidth) / Config.recommendedImageRectWidth;
+		const float HeightDensitymax = float(Config.maxImageRectHeight) / Config.recommendedImageRectHeight;
+		const float PerViewPixelDensityMax = FMath::Min(WidthDensityMax, HeightDensitymax);
+		RuntimePixelDensityMax = FMath::Min(RuntimePixelDensityMax, PerViewPixelDensityMax);
+	}
+
+	// Enumerate environment blend modes and select the best one.
+	{
+		TArray<XrEnvironmentBlendMode> BlendModes = RetrieveEnvironmentBlendModes();
+
+		// Select the first blend mode returned by the runtime that is supported.
+		// This is the environment blend mode preferred by the runtime.
+		for (XrEnvironmentBlendMode BlendMode : BlendModes)
+		{
+			if (IsEnvironmentBlendModeSupported(BlendMode))
+			{
+				SelectedEnvironmentBlendMode = BlendMode;
+				break;
+			}
+		}
+
+		// If there is no supported environment blend mode, use the first option as a last resort.
+		if (!ensure(SelectedEnvironmentBlendMode != XR_ENVIRONMENT_BLEND_MODE_MAX_ENUM))
+		{
+			SelectedEnvironmentBlendMode = BlendModes[0];
+		}
+	}
+
+	// Give the all frame states the same initial values.
+	PipelinedFrameStateRHI = PipelinedFrameStateRendering = PipelinedFrameStateGame;
+
 	XrSessionCreateInfo SessionInfo;
 	SessionInfo.type = XR_TYPE_SESSION_CREATE_INFO;
-	SessionInfo.next = RenderBridge.IsValid() ? RenderBridge->GetGraphicsBinding() : nullptr;
+	SessionInfo.next = nullptr;
+	if (RenderBridge.IsValid())
+	{
+		SessionInfo.next = RenderBridge->GetGraphicsBinding(System);
+		if (!SessionInfo.next)
+		{
+			UE_LOG(LogHMD, Warning, TEXT("Failed to get an OpenXR graphics binding, editor restart required."));
+#if WITH_EDITOR
+			ShowRestartWarning(LOCTEXT("EditorRestartMsg_Title", "Editor Restart Required"));
+#endif
+			return false;
+		}
+	}
 	SessionInfo.createFlags = 0;
 	SessionInfo.systemId = System;
 	for (IOpenXRExtensionPlugin* Module : ExtensionPlugins)
diff --git a/Engine/Plugins/Runtime/OpenXR/Source/OpenXRHMD/Private/OpenXRHMD.h b/Engine/Plugins/Runtime/OpenXR/Source/OpenXRHMD/Private/OpenXRHMD.h
--- a/Engine/Plugins/Runtime/OpenXR/Source/OpenXRHMD/Private/OpenXRHMD.h
+++ b/Engine/Plugins/Runtime/OpenXR/Source/OpenXRHMD/Private/OpenXRHMD.h
@@ -261,7 +261,7 @@ public:
 	virtual bool HasValidTrackingPosition() override { return IsTracking(HMDDeviceId); }
 
 	/** IHeadMountedDisplay interface */
-	virtual bool IsHMDConnected() override { return true; }
+	virtual bool IsHMDConnected() override;
 	virtual bool DoesSupportPositionalTracking() const override { return true; }
 	virtual bool IsHMDEnabled() const override;
 	virtual void EnableHMD(bool allow = true) override;
@@ -339,7 +339,7 @@ private:
 
 public:
 	/** Constructor */
-	FOpenXRHMD(const FAutoRegister&, XrInstance InInstance, XrSystemId InSystem, TRefCountPtr<FOpenXRRenderBridge>& InRenderBridge, TArray<const char*> InEnabledExtensions, TArray<class IOpenXRExtensionPlugin*> InExtensionPlugins, IARSystemSupport* ARSystemSupport);
+	FOpenXRHMD(const FAutoRegister&, XrInstance InInstance, TRefCountPtr<FOpenXRRenderBridge>& InRenderBridge, TArray<const char*> InEnabledExtensions, TArray<class IOpenXRExtensionPlugin*> InExtensionPlugins, IARSystemSupport* ARSystemSupport);
 
 	void SetInputModule(IOpenXRInputModule* InInputModule)
 	{
diff --git a/Engine/Plugins/Runtime/OpenXR/Source/OpenXRHMD/Private/OpenXRHMDModule.cpp b/Engine/Plugins/Runtime/OpenXR/Source/OpenXRHMD/Private/OpenXRHMDModule.cpp
--- a/Engine/Plugins/Runtime/OpenXR/Source/OpenXRHMD/Private/OpenXRHMDModule.cpp
+++ b/Engine/Plugins/Runtime/OpenXR/Source/OpenXRHMD/Private/OpenXRHMDModule.cpp
@@ -34,7 +34,6 @@ IMPLEMENT_MODULE( FOpenXRHMDModule, OpenXRHMD )
 FOpenXRHMDModule::FOpenXRHMDModule()
 	: LoaderHandle(nullptr)
 	, Instance(XR_NULL_HANDLE)
-	, System(XR_NULL_SYSTEM_ID)
 	, RenderBridge(nullptr)
 { }
 
@@ -44,11 +43,6 @@ FOpenXRHMDModule::~FOpenXRHMDModule()
 
 TSharedPtr< class IXRTrackingSystem, ESPMode::ThreadSafe > FOpenXRHMDModule::CreateTrackingSystem()
 {
-	if (!InitInstanceAndSystem())
-	{
-		return nullptr;
-	}
-
 	if (!RenderBridge && !FParse::Param(FCommandLine::Get(), TEXT("xrtrackingonly")))
 	{
 		if (!InitRenderBridge())
@@ -59,7 +53,7 @@ TSharedPtr< class IXRTrackingSystem, ESPMode::ThreadSafe > FOpenXRHMDModule::Cre
 	auto ARModule = FModuleManager::LoadModulePtr<IOpenXRARModule>("OpenXRAR");
 	auto ARSystem = ARModule->CreateARSystem();
 
-	auto OpenXRHMD = FSceneViewExtensions::NewExtension<FOpenXRHMD>(Instance, System, RenderBridge, EnabledExtensions, ExtensionPlugins, ARSystem);
+	auto OpenXRHMD = FSceneViewExtensions::NewExtension<FOpenXRHMD>(Instance, RenderBridge, EnabledExtensions, ExtensionPlugins, ARSystem);
 	if (OpenXRHMD->IsInitialized())
 	{
 		ARModule->SetTrackingSystem(OpenXRHMD);
@@ -74,7 +68,7 @@ void FOpenXRHMDModule::ShutdownModule()
 {
 	if (Instance)
 	{
-		DestroyInstance();
+		XR_ENSURE(xrDestroyInstance(Instance));
 	}
 
 	if (LoaderHandle)
@@ -98,14 +92,36 @@ uint64 FOpenXRHMDModule::GetGraphicsAdapterLuid()
 			return 0;
 		}
 	}
-	return RenderBridge->GetGraphicsAdapterLuid();
+
+	FConfigFile* EngineIni = GConfig->FindConfigFile(GEngineIni);
+	XrSystemId System = GetSystemId();
+	if (!System)
+	{
+		int64 AdapterLuid = 0;
+		EngineIni->GetInt64(TEXT("OpenXR.Settings"), TEXT("GraphicsAdapter"), AdapterLuid);
+		return reinterpret_cast<uint64&>(AdapterLuid);
+	}
+
+	uint64 AdapterLuid = RenderBridge->GetGraphicsAdapterLuid(System);
+	if (AdapterLuid)
+	{
+		// Remember this luid so we use the right adapter, even when we startup without an HMD connected
+		EngineIni->SetInt64(TEXT("OpenXR.Settings"), TEXT("GraphicsAdapter"), reinterpret_cast<int64&>(AdapterLuid));
+	}
+	return AdapterLuid;
 }
 
 TSharedPtr< IHeadMountedDisplayVulkanExtensions, ESPMode::ThreadSafe > FOpenXRHMDModule::GetVulkanExtensions()
 {
 #ifdef XR_USE_GRAPHICS_API_VULKAN
-	if (InitInstanceAndSystem() && IsExtensionEnabled(XR_KHR_VULKAN_ENABLE_EXTENSION_NAME))
+	if (InitInstance() && IsExtensionEnabled(XR_KHR_VULKAN_ENABLE_EXTENSION_NAME))
 	{
+		XrSystemId System = GetSystemId();
+		if (!System)
+		{
+			return nullptr;
+		}
+
 		if (!VulkanExtensions.IsValid())
 		{
 			VulkanExtensions = MakeShareable(new FOpenXRHMD::FVulkanExtensions(Instance, System));
@@ -118,21 +134,25 @@ TSharedPtr< IHeadMountedDisplayVulkanExtensions, ESPMode::ThreadSafe > FOpenXRHM
 
 FString FOpenXRHMDModule::GetDeviceSystemName()
 {
-	if (InitInstanceAndSystem())
+	if (InitInstance())
 	{
-		XrSystemProperties SystemProperties;
-		SystemProperties.type = XR_TYPE_SYSTEM_PROPERTIES;
-		SystemProperties.next = nullptr;
-		xrGetSystemProperties(Instance, System, &SystemProperties);
+		XrSystemId System = GetSystemId();
+		if (System)
+		{
+			XrSystemProperties SystemProperties;
+			SystemProperties.type = XR_TYPE_SYSTEM_PROPERTIES;
+			SystemProperties.next = nullptr;
+			XR_ENSURE(xrGetSystemProperties(Instance, System, &SystemProperties));
 
-		return FString(UTF8_TO_TCHAR(SystemProperties.systemName));
+			return FString(UTF8_TO_TCHAR(SystemProperties.systemName));
+		}
 	}
 	return FString("");
 }
 
 bool FOpenXRHMDModule::IsStandaloneStereoOnlyDevice()
 {
-	if (InitInstanceAndSystem())
+	if (InitInstance())
 	{
 		for (IOpenXRExtensionPlugin* Module : ExtensionPlugins)
 		{
@@ -251,7 +271,7 @@ bool FOpenXRHMDModule::InitRenderBridge()
 	for (IOpenXRExtensionPlugin* Plugin : ExtModules)
 	{
 		// We are taking ownership of the CustomRenderBridge instance here.
-		TRefCountPtr<FOpenXRRenderBridge> CustomRenderBridge = Plugin->GetCustomRenderBridge(Instance, System);
+		TRefCountPtr<FOpenXRRenderBridge> CustomRenderBridge = Plugin->GetCustomRenderBridge(Instance);
 		if (CustomRenderBridge)
 		{
 			// We pick the first
@@ -265,7 +285,7 @@ bool FOpenXRHMDModule::InitRenderBridge()
 		return false;
 	}
 
-	if (!InitInstanceAndSystem())
+	if (!InitInstance())
 	{
 		return false;
 	}
@@ -275,35 +295,35 @@ bool FOpenXRHMDModule::InitRenderBridge()
 #ifdef XR_USE_GRAPHICS_API_D3D11
 	if (RHIType == ERHIInterfaceType::D3D11 && IsExtensionEnabled(XR_KHR_D3D11_ENABLE_EXTENSION_NAME))
 	{
-		RenderBridge = CreateRenderBridge_D3D11(Instance, System);
+		RenderBridge = CreateRenderBridge_D3D11(Instance);
 	}
 	else
 #endif
 #ifdef XR_USE_GRAPHICS_API_D3D12
 	if (RHIType == ERHIInterfaceType::D3D12 && IsExtensionEnabled(XR_KHR_D3D12_ENABLE_EXTENSION_NAME))
 	{
-		RenderBridge = CreateRenderBridge_D3D12(Instance, System);
+		RenderBridge = CreateRenderBridge_D3D12(Instance);
 	}
 	else
 #endif
 #if defined(XR_USE_GRAPHICS_API_OPENGL_ES) && defined(XR_USE_PLATFORM_ANDROID)
 	if (RHIType == ERHIInterfaceType::OpenGL && IsExtensionEnabled(XR_KHR_OPENGL_ES_ENABLE_EXTENSION_NAME))
 	{
-		RenderBridge = CreateRenderBridge_OpenGLES(Instance, System);
+		RenderBridge = CreateRenderBridge_OpenGLES(Instance);
 	}
 	else
 #endif
 #ifdef XR_USE_GRAPHICS_API_OPENGL
 	if (RHIType == ERHIInterfaceType::OpenGL && IsExtensionEnabled(XR_KHR_OPENGL_ENABLE_EXTENSION_NAME))
 	{
-		RenderBridge = CreateRenderBridge_OpenGL(Instance, System);
+		RenderBridge = CreateRenderBridge_OpenGL(Instance);
 	}
 	else
 #endif
 #ifdef XR_USE_GRAPHICS_API_VULKAN
 	if (RHIType == ERHIInterfaceType::Vulkan && IsExtensionEnabled(XR_KHR_VULKAN_ENABLE_EXTENSION_NAME))
 	{
-		RenderBridge = CreateRenderBridge_Vulkan(Instance, System);
+		RenderBridge = CreateRenderBridge_Vulkan(Instance);
 	}
 	else
 #endif
@@ -473,26 +493,13 @@ bool FOpenXRHMDModule::GetOptionalExtensions(TArray<const ANSICHAR*>& OutExtensi
 	return true;
 }
 
-bool FOpenXRHMDModule::InitInstanceAndSystem()
+bool FOpenXRHMDModule::InitInstance()
 {
-	if (!Instance && !InitInstance())
-	{
-		return false;
-	}
-
-	if (!System && !InitSystem())
+	if (Instance)
 	{
-		return false;
+		return true;
 	}
 
-	return true;
-}
-
-bool FOpenXRHMDModule::InitInstance()
-{
-	// This should only ever be called if we don't already have an instance.
-	check(!Instance);
-
 #if PLATFORM_ANDROID
 	// TODO: Allow OpenXR on non-Oculus Android platforms
 	if (AndroidThunkCpp_IsOculusMobileApplication())
@@ -731,8 +738,10 @@ bool FOpenXRHMDModule::InitInstance()
 	return true;
 }
 
-bool FOpenXRHMDModule::InitSystem()
+XrSystemId FOpenXRHMDModule::GetSystemId() const
 {
+	XrSystemId System = XR_NULL_SYSTEM_ID;
+
 	XrSystemGetInfo SystemInfo;
 	SystemInfo.type = XR_TYPE_SYSTEM_GET_INFO;
 	SystemInfo.next = nullptr;
@@ -745,9 +754,8 @@ bool FOpenXRHMDModule::InitSystem()
 	XrResult Result = xrGetSystem(Instance, &SystemInfo, &System);
 	if (XR_FAILED(Result))
 	{
-		DestroyInstance();
-		UE_LOG(LogHMD, Log, TEXT("Failed to get an OpenXR system, result is %s. Please check that your runtime supports VR headsets."), OpenXRResultToString(Result));
-		return false;
+		UE_LOG(LogHMD, VeryVerbose, TEXT("Failed to get an OpenXR system, result is %s"), OpenXRResultToString(Result));
+		return XR_NULL_SYSTEM_ID;
 	}
 
 	for (IOpenXRExtensionPlugin* Module : ExtensionPlugins)
@@ -755,13 +763,7 @@ bool FOpenXRHMDModule::InitSystem()
 		Module->PostGetSystem(Instance, System);
 	}
 
-	return true;
-}
-
-void FOpenXRHMDModule::DestroyInstance()
-{
-	XR_ENSURE(xrDestroyInstance(Instance));
-	Instance = XR_NULL_HANDLE;
+	return System;
 }
 
 FName FOpenXRHMDModule::ResolvePathToName(XrPath Path)
diff --git a/Engine/Plugins/Runtime/OpenXR/Source/OpenXRHMD/Private/OpenXRHMDModule.h b/Engine/Plugins/Runtime/OpenXR/Source/OpenXRHMD/Private/OpenXRHMDModule.h
--- a/Engine/Plugins/Runtime/OpenXR/Source/OpenXRHMD/Private/OpenXRHMDModule.h
+++ b/Engine/Plugins/Runtime/OpenXR/Source/OpenXRHMD/Private/OpenXRHMDModule.h
@@ -39,7 +39,7 @@ public:
 	virtual bool IsLayerAvailable(const FString& Name) const override { return EnabledLayers.Contains(Name); }
 	virtual bool IsLayerEnabled(const FString& Name) const override { return EnabledLayers.Contains(Name); }
 	virtual XrInstance GetInstance() const override { return Instance; }
-	virtual XrSystemId GetSystemId() const override { return System; }
+	virtual XrSystemId GetSystemId() const override;
 
 	virtual FName ResolvePathToName(XrPath Path) override;
 	virtual XrPath ResolveNameToPath(FName Name) override;
@@ -47,7 +47,6 @@ public:
 private:
 	void* LoaderHandle;
 	XrInstance Instance;
-	XrSystemId System;
 	TSet<FString> AvailableExtensions;
 	TSet<FString> AvailableLayers;
 	TArray<const char*> EnabledExtensions;
@@ -61,13 +60,10 @@ private:
 	TSortedMap<XrPath, FName> PathToName;
 	TSortedMap<FName, XrPath, FDefaultAllocator, FNameFastLess> NameToPath;
 
-	void DestroyInstance();
 	bool EnumerateExtensions();
 	bool EnumerateLayers();
 	bool InitRenderBridge();
-	bool InitInstanceAndSystem();
 	bool InitInstance();
-	bool InitSystem();
 	PFN_xrGetInstanceProcAddr GetDefaultLoader();
 	bool EnableExtensions(const TArray<const ANSICHAR*>& RequiredExtensions, const TArray<const ANSICHAR*>& OptionalExtensions, TArray<const ANSICHAR*>& OutExtensions);
 	bool GetRequiredExtensions(TArray<const ANSICHAR*>& OutExtensions);
diff --git a/Engine/Plugins/Runtime/OpenXR/Source/OpenXRHMD/Private/OpenXRHMD_RenderBridge.cpp b/Engine/Plugins/Runtime/OpenXR/Source/OpenXRHMD/Private/OpenXRHMD_RenderBridge.cpp
--- a/Engine/Plugins/Runtime/OpenXR/Source/OpenXRHMD/Private/OpenXRHMD_RenderBridge.cpp
+++ b/Engine/Plugins/Runtime/OpenXR/Source/OpenXRHMD/Private/OpenXRHMD_RenderBridge.cpp
@@ -24,23 +24,34 @@ bool FOpenXRRenderBridge::Present(int32& InOutSyncInterval)
 class FD3D11RenderBridge : public FOpenXRRenderBridge
 {
 public:
-	FD3D11RenderBridge(XrInstance InInstance, XrSystemId InSystem)
-		: Binding()
+	FD3D11RenderBridge(XrInstance InInstance)
+		: FOpenXRRenderBridge(InInstance)
+		, Binding()
 	{
-		PFN_xrGetD3D11GraphicsRequirementsKHR GetD3D11GraphicsRequirementsKHR;
 		XR_ENSURE(xrGetInstanceProcAddr(InInstance, "xrGetD3D11GraphicsRequirementsKHR", (PFN_xrVoidFunction*)&GetD3D11GraphicsRequirementsKHR));
+	}
 
+	virtual uint64 GetGraphicsAdapterLuid(XrSystemId InSystem) override
+	{
 		XrGraphicsRequirementsD3D11KHR Requirements;
 		Requirements.type = XR_TYPE_GRAPHICS_REQUIREMENTS_D3D11_KHR;
 		Requirements.next = nullptr;
-		if (XR_ENSURE(GetD3D11GraphicsRequirementsKHR(InInstance, InSystem, &Requirements)))
+		if (XR_ENSURE(GetD3D11GraphicsRequirementsKHR(Instance, InSystem, &Requirements)))
 		{
-			AdapterLuid = reinterpret_cast<uint64&>(Requirements.adapterLuid);
+			return reinterpret_cast<uint64&>(Requirements.adapterLuid);
 		}
+		return 0;
 	}
 
-	virtual void* GetGraphicsBinding() override
+	virtual void* GetGraphicsBinding(XrSystemId InSystem) override
 	{
+		DXGI_ADAPTER_DESC AdapterDesc;
+		GetID3D11DynamicRHI()->RHIGetAdapter()->GetDesc(&AdapterDesc);
+		if (reinterpret_cast<uint64&>(AdapterDesc.AdapterLuid) != GetGraphicsAdapterLuid(InSystem))
+		{
+			return nullptr;
+		}
+
 		Binding.type = XR_TYPE_GRAPHICS_BINDING_D3D11_KHR;
 		Binding.next = nullptr;
 		Binding.device = GetID3D11DynamicRHI()->RHIGetDevice();
@@ -53,33 +64,44 @@ public:
 	}
 
 private:
+	PFN_xrGetD3D11GraphicsRequirementsKHR GetD3D11GraphicsRequirementsKHR;
 	XrGraphicsBindingD3D11KHR Binding;
 };
 
-FOpenXRRenderBridge* CreateRenderBridge_D3D11(XrInstance InInstance, XrSystemId InSystem) { return new FD3D11RenderBridge(InInstance, InSystem); }
+FOpenXRRenderBridge* CreateRenderBridge_D3D11(XrInstance InInstance) { return new FD3D11RenderBridge(InInstance); }
 #endif
 
 #ifdef XR_USE_GRAPHICS_API_D3D12
 class FD3D12RenderBridge : public FOpenXRRenderBridge
 {
 public:
-	FD3D12RenderBridge(XrInstance InInstance, XrSystemId InSystem)
-		: Binding()
+	FD3D12RenderBridge(XrInstance InInstance)
+		: FOpenXRRenderBridge(InInstance)
+		, Binding()
 	{
-		PFN_xrGetD3D12GraphicsRequirementsKHR GetD3D12GraphicsRequirementsKHR;
 		XR_ENSURE(xrGetInstanceProcAddr(InInstance, "xrGetD3D12GraphicsRequirementsKHR", (PFN_xrVoidFunction*)&GetD3D12GraphicsRequirementsKHR));
+	}
 
+	virtual uint64 GetGraphicsAdapterLuid(XrSystemId InSystem) override
+	{
 		XrGraphicsRequirementsD3D12KHR Requirements;
 		Requirements.type = XR_TYPE_GRAPHICS_REQUIREMENTS_D3D12_KHR;
 		Requirements.next = nullptr;
-		if (XR_ENSURE(GetD3D12GraphicsRequirementsKHR(InInstance, InSystem, &Requirements)))
+		if (XR_ENSURE(GetD3D12GraphicsRequirementsKHR(Instance, InSystem, &Requirements)))
 		{
-			AdapterLuid = reinterpret_cast<uint64&>(Requirements.adapterLuid);
+			return reinterpret_cast<uint64&>(Requirements.adapterLuid);
 		}
+		return 0;
 	}
 
-	virtual void* GetGraphicsBinding() override
+	virtual void* GetGraphicsBinding(XrSystemId InSystem) override
 	{
+		LUID AdapterLuid = GetID3D12DynamicRHI()->RHIGetDevice(0)->GetAdapterLuid();
+		if (reinterpret_cast<uint64&>(AdapterLuid) != GetGraphicsAdapterLuid(InSystem))
+		{
+			return nullptr;
+		}
+
 		Binding.type = XR_TYPE_GRAPHICS_BINDING_D3D12_KHR;
 		Binding.next = nullptr;
 		Binding.device = GetID3D12DynamicRHI()->RHIGetDevice(0);
@@ -93,28 +115,32 @@ public:
 	}
 
 private:
+	PFN_xrGetD3D12GraphicsRequirementsKHR GetD3D12GraphicsRequirementsKHR;
 	XrGraphicsBindingD3D12KHR Binding;
 };
 
-FOpenXRRenderBridge* CreateRenderBridge_D3D12(XrInstance InInstance, XrSystemId InSystem) { return new FD3D12RenderBridge(InInstance, InSystem); }
+FOpenXRRenderBridge* CreateRenderBridge_D3D12(XrInstance InInstance) { return new FD3D12RenderBridge(InInstance); }
 #endif
 
 #ifdef XR_USE_GRAPHICS_API_OPENGL
 class FOpenGLRenderBridge : public FOpenXRRenderBridge
 {
 public:
-	FOpenGLRenderBridge(XrInstance InInstance, XrSystemId InSystem)
-		: Binding()
+	FOpenGLRenderBridge(XrInstance InInstance)
+		: FOpenXRRenderBridge(InInstance)
+		, Binding()
 	{
-		PFN_xrGetOpenGLGraphicsRequirementsKHR GetOpenGLGraphicsRequirementsKHR;
 		XR_ENSURE(xrGetInstanceProcAddr(InInstance, "xrGetOpenGLGraphicsRequirementsKHR", (PFN_xrVoidFunction*)&GetOpenGLGraphicsRequirementsKHR));
+	}
 
+	virtual void* GetGraphicsBinding(XrSystemId InSystem) override
+	{
 		XrGraphicsRequirementsOpenGLKHR Requirements;
 		Requirements.type = XR_TYPE_GRAPHICS_REQUIREMENTS_OPENGL_KHR;
 		Requirements.next = nullptr;
 		Requirements.minApiVersionSupported = 0;
 		Requirements.maxApiVersionSupported = 0;
-		XR_ENSURE(GetOpenGLGraphicsRequirementsKHR(InInstance, InSystem, &Requirements));
+		XR_ENSURE(GetOpenGLGraphicsRequirementsKHR(Instance, InSystem, &Requirements));
 
 		IOpenGLDynamicRHI* RHI = GetIOpenGLDynamicRHI();
 		XrVersion RHIVersion = XR_MAKE_VERSION(RHI->RHIGetGLMajorVersion(), RHI->RHIGetGLMinorVersion(), 0);
@@ -127,10 +153,7 @@ public:
 		{
 			UE_LOG(LogHMD, Warning, TEXT("The OpenGL API version has not been tested with the OpenXR runtime"));
 		}
-	}
 
-	virtual void* GetGraphicsBinding() override
-	{
 #if PLATFORM_WINDOWS
 		Binding.type = XR_TYPE_GRAPHICS_BINDING_OPENGL_WIN32_KHR;
 		Binding.next = nullptr;
@@ -147,22 +170,23 @@ public:
 	}
 
 private:
+	PFN_xrGetOpenGLGraphicsRequirementsKHR GetOpenGLGraphicsRequirementsKHR;
 #if PLATFORM_WINDOWS
 	XrGraphicsBindingOpenGLWin32KHR Binding;
 #endif
 };
 
-FOpenXRRenderBridge* CreateRenderBridge_OpenGL(XrInstance InInstance, XrSystemId InSystem) { return new FOpenGLRenderBridge(InInstance, InSystem); }
+FOpenXRRenderBridge* CreateRenderBridge_OpenGL(XrInstance InInstance) { return new FOpenGLRenderBridge(InInstance); }
 #endif
 
 #ifdef XR_USE_GRAPHICS_API_OPENGL_ES
 class FOpenGLESRenderBridge : public FOpenXRRenderBridge
 {
 public:
-	FOpenGLESRenderBridge(XrInstance InInstance, XrSystemId InSystem)
-		: Binding()
+	FOpenGLESRenderBridge(XrInstance InInstance)
+		: FOpenXRRenderBridge(InInstance)
+		, Binding()
 	{
-		PFN_xrGetOpenGLESGraphicsRequirementsKHR GetOpenGLESGraphicsRequirementsKHR;
 		XR_ENSURE(xrGetInstanceProcAddr(InInstance, "xrGetOpenGLESGraphicsRequirementsKHR", (PFN_xrVoidFunction*)&GetOpenGLESGraphicsRequirementsKHR));
 
 		XrGraphicsRequirementsOpenGLESKHR Requirements;
@@ -187,7 +211,7 @@ public:
 #endif
 	}
 
-	virtual void* GetGraphicsBinding() override
+	virtual void* GetGraphicsBinding(XrSystemId InSystem) override
 	{
 #if PLATFORM_ANDROID
 		IOpenGLDynamicRHI* RHI = GetIOpenGLDynamicRHI();
@@ -207,32 +231,35 @@ public:
 	}
 
 private:
+	PFN_xrGetOpenGLESGraphicsRequirementsKHR GetOpenGLESGraphicsRequirementsKHR;
 #if PLATFORM_ANDROID
 	XrGraphicsBindingOpenGLESAndroidKHR Binding;
 #endif
 };
 
-FOpenXRRenderBridge* CreateRenderBridge_OpenGLES(XrInstance InInstance, XrSystemId InSystem) { return new FOpenGLESRenderBridge(InInstance, InSystem); }
+FOpenXRRenderBridge* CreateRenderBridge_OpenGLES(XrInstance InInstance) { return new FOpenGLESRenderBridge(InInstance); }
 #endif
 
 #ifdef XR_USE_GRAPHICS_API_VULKAN
 class FVulkanRenderBridge : public FOpenXRRenderBridge
 {
 public:
-	FVulkanRenderBridge(XrInstance InInstance, XrSystemId InSystem)
-		: Binding()
-		, Instance(InInstance)
-		, System(InSystem)
+	FVulkanRenderBridge(XrInstance InInstance)
+		: FOpenXRRenderBridge(InInstance)
+		, Binding()
 	{
-		PFN_xrGetVulkanGraphicsRequirementsKHR GetVulkanGraphicsRequirementsKHR;
 		XR_ENSURE(xrGetInstanceProcAddr(InInstance, "xrGetVulkanGraphicsRequirementsKHR", (PFN_xrVoidFunction*)&GetVulkanGraphicsRequirementsKHR));
+		XR_ENSURE(xrGetInstanceProcAddr(InInstance, "xrGetVulkanGraphicsDeviceKHR", (PFN_xrVoidFunction*)&GetVulkanGraphicsDeviceKHR));
+	}
 
+	virtual void* GetGraphicsBinding(XrSystemId InSystem) override
+	{
 		XrGraphicsRequirementsVulkanKHR Requirements;
 		Requirements.type = XR_TYPE_GRAPHICS_REQUIREMENTS_VULKAN_KHR;
 		Requirements.next = nullptr;
 		Requirements.minApiVersionSupported = 0;
 		Requirements.maxApiVersionSupported = 0;
-		XR_ENSURE(GetVulkanGraphicsRequirementsKHR(InInstance, InSystem, &Requirements));
+		XR_ENSURE(GetVulkanGraphicsRequirementsKHR(Instance, InSystem, &Requirements));
 
 		IVulkanDynamicRHI* VulkanRHI = GetIVulkanDynamicRHI();
 		const uint32 VulkanVersion = VulkanRHI->RHIGetVulkanVersion();
@@ -253,14 +280,12 @@ public:
 			UE_LOG(LogHMD, Warning, TEXT("The Vulkan API version has not been tested with the OpenXR runtime"));
 		}
 
-		PFN_xrGetVulkanGraphicsDeviceKHR GetVulkanGraphicsDeviceKHR;
-		XR_ENSURE(xrGetInstanceProcAddr(Instance, "xrGetVulkanGraphicsDeviceKHR", (PFN_xrVoidFunction*)&GetVulkanGraphicsDeviceKHR));
-		XR_ENSURE(GetVulkanGraphicsDeviceKHR(Instance, System, VulkanRHI->RHIGetVkInstance(), &Gpu));
-	}
-
-	virtual void* GetGraphicsBinding() override
-	{
-		IVulkanDynamicRHI* VulkanRHI = GetIVulkanDynamicRHI();
+		VkPhysicalDevice Gpu;
+		XR_ENSURE(GetVulkanGraphicsDeviceKHR(Instance, InSystem, VulkanRHI->RHIGetVkInstance(), &Gpu));
+		if (Gpu != VulkanRHI->RHIGetVkPhysicalDevice())
+		{
+			return nullptr;
+		}
 
 		Binding.type = XR_TYPE_GRAPHICS_BINDING_VULKAN_KHR;
 		Binding.next = nullptr;
@@ -272,14 +297,11 @@ public:
 		return &Binding;
 	}
 
-	virtual uint64 GetGraphicsAdapterLuid() override
+	virtual uint64 GetGraphicsAdapterLuid(XrSystemId InSystem) override
 	{
-		if (!AdapterLuid)
-		{
-			AdapterLuid = GetIVulkanDynamicRHI()->RHIGetGraphicsAdapterLUID(Gpu);
-		}
-
-		return AdapterLuid;
+		VkPhysicalDevice Gpu;
+		XR_ENSURE(GetVulkanGraphicsDeviceKHR(Instance, InSystem, GetIVulkanDynamicRHI()->RHIGetVkInstance(), &Gpu));
+		return GetIVulkanDynamicRHI()->RHIGetGraphicsAdapterLUID(Gpu);
 	}
 
 	virtual FXRSwapChainPtr CreateSwapchain(XrSession InSession, uint8 Format, uint8& OutActualFormat, uint32 SizeX, uint32 SizeY, uint32 ArraySize, uint32 NumMips, uint32 NumSamples, ETextureCreateFlags CreateFlags, const FClearValueBinding& ClearValueBinding) override final
@@ -288,11 +310,10 @@ public:
 	}
 
 private:
+	PFN_xrGetVulkanGraphicsRequirementsKHR GetVulkanGraphicsRequirementsKHR;
+	PFN_xrGetVulkanGraphicsDeviceKHR GetVulkanGraphicsDeviceKHR;
 	XrGraphicsBindingVulkanKHR Binding;
-	XrInstance Instance;
-	XrSystemId System;
-	VkPhysicalDevice Gpu;
 };
 
-FOpenXRRenderBridge* CreateRenderBridge_Vulkan(XrInstance InInstance, XrSystemId InSystem) { return new FVulkanRenderBridge(InInstance, InSystem); }
+FOpenXRRenderBridge* CreateRenderBridge_Vulkan(XrInstance InInstance) { return new FVulkanRenderBridge(InInstance); }
 #endif
diff --git a/Engine/Plugins/Runtime/OpenXR/Source/OpenXRHMD/Private/OpenXRHMD_RenderBridge.h b/Engine/Plugins/Runtime/OpenXR/Source/OpenXRHMD/Private/OpenXRHMD_RenderBridge.h
--- a/Engine/Plugins/Runtime/OpenXR/Source/OpenXRHMD/Private/OpenXRHMD_RenderBridge.h
+++ b/Engine/Plugins/Runtime/OpenXR/Source/OpenXRHMD/Private/OpenXRHMD_RenderBridge.h
@@ -16,15 +16,15 @@ class FRHICommandListImmediate;
 class FOpenXRRenderBridge : public FXRRenderBridge
 {
 public:
-	FOpenXRRenderBridge()
-		: AdapterLuid(0)
+	FOpenXRRenderBridge(XrInstance InInstance)
+		: Instance(InInstance)
 		, OpenXRHMD(nullptr)
 	{ }
 
 	void SetOpenXRHMD(FOpenXRHMD* InHMD) { OpenXRHMD = InHMD; }
-	virtual uint64 GetGraphicsAdapterLuid() { return AdapterLuid; }
+	virtual uint64 GetGraphicsAdapterLuid(XrSystemId InSystem) { return 0; };
+	virtual void* GetGraphicsBinding(XrSystemId InSystem) = 0;
 
-	virtual void* GetGraphicsBinding() = 0;
 	virtual FXRSwapChainPtr CreateSwapchain(XrSession InSession, uint8 Format, uint8& OutActualFormat, uint32 SizeX, uint32 SizeY, uint32 ArraySize, uint32 NumMips, uint32 NumSamples, ETextureCreateFlags CreateFlags, const FClearValueBinding& ClearValueBinding) = 0;
 
 	FXRSwapChainPtr CreateSwapchain(XrSession InSession, FRHITexture2D* Template, ETextureCreateFlags CreateFlags)
@@ -71,27 +71,27 @@ public:
 	}
 
 protected:
-	uint64 AdapterLuid;
 	bool bNativePresent_RenderThread = true;
 	bool bNativePresent_RHIThread = true;
 	
+	XrInstance Instance;
 	FOpenXRHMD* OpenXRHMD;
 
 private:
 };
 
 #ifdef XR_USE_GRAPHICS_API_D3D11
-FOpenXRRenderBridge* CreateRenderBridge_D3D11(XrInstance InInstance, XrSystemId InSystem);
+FOpenXRRenderBridge* CreateRenderBridge_D3D11(XrInstance InInstance);
 #endif
 #ifdef XR_USE_GRAPHICS_API_D3D12
-FOpenXRRenderBridge* CreateRenderBridge_D3D12(XrInstance InInstance, XrSystemId InSystem);
+FOpenXRRenderBridge* CreateRenderBridge_D3D12(XrInstance InInstance);
 #endif
 #ifdef XR_USE_GRAPHICS_API_OPENGL_ES
-FOpenXRRenderBridge* CreateRenderBridge_OpenGLES(XrInstance InInstance, XrSystemId InSystem);
+FOpenXRRenderBridge* CreateRenderBridge_OpenGLES(XrInstance InInstance);
 #endif
 #ifdef XR_USE_GRAPHICS_API_OPENGL
-FOpenXRRenderBridge* CreateRenderBridge_OpenGL(XrInstance InInstance, XrSystemId InSystem);
+FOpenXRRenderBridge* CreateRenderBridge_OpenGL(XrInstance InInstance);
 #endif
 #ifdef XR_USE_GRAPHICS_API_VULKAN
-FOpenXRRenderBridge* CreateRenderBridge_Vulkan(XrInstance InInstance, XrSystemId InSystem);
+FOpenXRRenderBridge* CreateRenderBridge_Vulkan(XrInstance InInstance);
 #endif
diff --git a/Engine/Plugins/Runtime/OpenXR/Source/OpenXRHMD/Public/IOpenXRExtensionPlugin.h b/Engine/Plugins/Runtime/OpenXR/Source/OpenXRHMD/Public/IOpenXRExtensionPlugin.h
--- a/Engine/Plugins/Runtime/OpenXR/Source/OpenXRHMD/Public/IOpenXRExtensionPlugin.h
+++ b/Engine/Plugins/Runtime/OpenXR/Source/OpenXRHMD/Public/IOpenXRExtensionPlugin.h
@@ -172,7 +172,7 @@ public:
 	* Optionally provide a custom render bridge for the OpenXR plugin.
 	* Note: this returns a pointer to a new instance allocated with "new".  Calling code is responsible for eventually deleting it.
 	*/
-	virtual class FOpenXRRenderBridge* GetCustomRenderBridge(XrInstance InInstance, XrSystemId InSystem)
+	virtual class FOpenXRRenderBridge* GetCustomRenderBridge(XrInstance InInstance)
 	{
 		return nullptr;
 	}
diff --git a/Engine/Source/Editor/UnrealEd/Private/Kismet2/DebuggerCommands.cpp b/Engine/Source/Editor/UnrealEd/Private/Kismet2/DebuggerCommands.cpp
--- a/Engine/Source/Editor/UnrealEd/Private/Kismet2/DebuggerCommands.cpp
+++ b/Engine/Source/Editor/UnrealEd/Private/Kismet2/DebuggerCommands.cpp
@@ -1601,7 +1601,8 @@ bool FInternalPlayWorldCommandCallbacks::PlayInVR_CanExecute()
 	{
 		return false;
 	}
-	return (!HasPlayWorld() || !GUnrealEd->bIsSimulatingInEditor) && !GEditor->IsLightingBuildCurrentlyRunning() && GEngine && GEngine->XRSystem.IsValid();
+	return (!HasPlayWorld() || !GUnrealEd->bIsSimulatingInEditor) && !GEditor->IsLightingBuildCurrentlyRunning() &&
+		GEngine && GEngine->XRSystem.IsValid() && GEngine->XRSystem->GetHMDDevice() && GEngine->XRSystem->GetHMDDevice()->IsHMDConnected();
 }
 
 void SetLastExecutedPIEPreviewDevice(FString PIEPreviewDevice)
