Fix compile errors in CL 23781334

Also downgrade a fatal log to an error now that there's higher-level error handling.

diff --git a/Engine/Plugins/Runtime/OpenXR/Source/OpenXRHMD/Private/OpenXRHMD_RenderBridge.cpp b/Engine/Plugins/Runtime/OpenXR/Source/OpenXRHMD/Private/OpenXRHMD_RenderBridge.cpp
--- a/Engine/Plugins/Runtime/OpenXR/Source/OpenXRHMD/Private/OpenXRHMD_RenderBridge.cpp
+++ b/Engine/Plugins/Runtime/OpenXR/Source/OpenXRHMD/Private/OpenXRHMD_RenderBridge.cpp
@@ -146,7 +146,8 @@ public:
 		XrVersion RHIVersion = XR_MAKE_VERSION(RHI->RHIGetGLMajorVersion(), RHI->RHIGetGLMinorVersion(), 0);
 		if (RHIVersion < Requirements.minApiVersionSupported) //-V547
 		{
-			UE_LOG(LogHMD, Fatal, TEXT("The OpenGL API version does not meet the minimum version required by the OpenXR runtime"));
+			UE_LOG(LogHMD, Error, TEXT("The OpenGL API version does not meet the minimum version required by the OpenXR runtime"));
+			return nullptr;
 		}
 
 		if (RHIVersion > Requirements.maxApiVersionSupported) //-V547
@@ -188,33 +189,31 @@ public:
 		, Binding()
 	{
 		XR_ENSURE(xrGetInstanceProcAddr(InInstance, "xrGetOpenGLESGraphicsRequirementsKHR", (PFN_xrVoidFunction*)&GetOpenGLESGraphicsRequirementsKHR));
+	}
 
+	virtual void* GetGraphicsBinding(XrSystemId InSystem) override
+	{
 		XrGraphicsRequirementsOpenGLESKHR Requirements;
 		Requirements.type = XR_TYPE_GRAPHICS_REQUIREMENTS_OPENGL_ES_KHR;
 		Requirements.next = nullptr;
 		Requirements.minApiVersionSupported = 0;
 		Requirements.maxApiVersionSupported = 0;
-		XR_ENSURE(GetOpenGLESGraphicsRequirementsKHR(InInstance, InSystem, &Requirements));
+		XR_ENSURE(GetOpenGLESGraphicsRequirementsKHR(Instance, InSystem, &Requirements));
 
 #if PLATFORM_ANDROID
 		IOpenGLDynamicRHI* RHI = GetIOpenGLDynamicRHI();
 		XrVersion RHIVersion = XR_MAKE_VERSION(RHI->RHIGetGLMajorVersion(), RHI->RHIGetGLMinorVersion(), 0);
 		if (RHIVersion < Requirements.minApiVersionSupported) //-V547
 		{
-			UE_LOG(LogHMD, Fatal, TEXT("The OpenGLES API version does not meet the minimum version required by the OpenXR runtime"));
+			UE_LOG(LogHMD, Error, TEXT("The OpenGLES API version does not meet the minimum version required by the OpenXR runtime"));
+			return nullptr;
 		}
 
 		if (RHIVersion > Requirements.maxApiVersionSupported) //-V547
 		{
 			UE_LOG(LogHMD, Warning, TEXT("The OpenGLES API version has not been tested with the OpenXR runtime"));
 		}
-#endif
-	}
 
-	virtual void* GetGraphicsBinding(XrSystemId InSystem) override
-	{
-#if PLATFORM_ANDROID
-		IOpenGLDynamicRHI* RHI = GetIOpenGLDynamicRHI();
 		Binding.type = XR_TYPE_GRAPHICS_BINDING_OPENGL_ES_ANDROID_KHR;
 		Binding.next = nullptr;
 		Binding.display = RHI->RHIGetEGLDisplay();
