Fixed missing undef LOCTEXT_NAMESPACE

diff --git a/Engine/Plugins/Runtime/OpenXR/Source/OpenXRHMD/Private/OpenXRHMD.cpp b/Engine/Plugins/Runtime/OpenXR/Source/OpenXRHMD/Private/OpenXRHMD.cpp
--- a/Engine/Plugins/Runtime/OpenXR/Source/OpenXRHMD/Private/OpenXRHMD.cpp
+++ b/Engine/Plugins/Runtime/OpenXR/Source/OpenXRHMD/Private/OpenXRHMD.cpp
@@ -3601,3 +3601,5 @@ void FOpenXRHMD::FTrackingSpace::DestroySpace()
 	}
 	Handle = XR_NULL_HANDLE;
 }
+
+#undef LOCTEXT_NAMESPACE
\ No newline at end of file
