Misc static analysis fixes.

#okfotgithub public

diff --git a/Engine/Plugins/Runtime/HairStrands/Source/HairStrandsCore/Private/GroomBuilder.cpp b/Engine/Plugins/Runtime/HairStrands/Source/HairStrandsCore/Private/GroomBuilder.cpp
--- a/Engine/Plugins/Runtime/HairStrands/Source/HairStrandsCore/Private/GroomBuilder.cpp
+++ b/Engine/Plugins/Runtime/HairStrands/Source/HairStrandsCore/Private/GroomBuilder.cpp
@@ -1675,7 +1675,7 @@ bool FGroomBuilder::BuildHairDescriptionGroups(const FHairDescription& HairDescr
 		}
 		FHairDescriptionGroup& Group = Out.HairGroups.AddDefaulted_GetRef();
 		Group.Info.GroupID = GroupID;
-		Group.Info.GroupName = GroupName != NAME_None ? GroupName : FName(FString::Printf(TEXT("Group_%d"), GroupID));
+		Group.Info.GroupName = (GroupName != NAME_None) ? GroupName : FName(FString::Printf(TEXT("Group_%d"), GroupID));
 		Group.Info.MaxImportedWidth = -1.0f; // Invalid width
 		return Group;
 	};
diff --git a/Engine/Plugins/Runtime/OpenXR/Source/OpenXRHMD/Private/OpenXRHMD_RenderBridge.cpp b/Engine/Plugins/Runtime/OpenXR/Source/OpenXRHMD/Private/OpenXRHMD_RenderBridge.cpp
--- a/Engine/Plugins/Runtime/OpenXR/Source/OpenXRHMD/Private/OpenXRHMD_RenderBridge.cpp
+++ b/Engine/Plugins/Runtime/OpenXR/Source/OpenXRHMD/Private/OpenXRHMD_RenderBridge.cpp
@@ -279,7 +279,7 @@ public:
 			UE_LOG(LogHMD, Warning, TEXT("The Vulkan API version has not been tested with the OpenXR runtime"));
 		}
 
-		VkPhysicalDevice Gpu;
+		VkPhysicalDevice Gpu = nullptr;
 		XR_ENSURE(GetVulkanGraphicsDeviceKHR(Instance, InSystem, VulkanRHI->RHIGetVkInstance(), &Gpu));
 		if (Gpu != VulkanRHI->RHIGetVkPhysicalDevice())
 		{
@@ -298,7 +298,7 @@ public:
 
 	virtual uint64 GetGraphicsAdapterLuid(XrSystemId InSystem) override
 	{
-		VkPhysicalDevice Gpu;
+		VkPhysicalDevice Gpu = nullptr;
 		XR_ENSURE(GetVulkanGraphicsDeviceKHR(Instance, InSystem, GetIVulkanDynamicRHI()->RHIGetVkInstance(), &Gpu));
 		return GetIVulkanDynamicRHI()->RHIGetGraphicsAdapterLUID(Gpu);
 	}
