OpenXR: Temp fix for crash when nDisplay launches, by avoid a null dereferencing of the XrInstance object.

#rnx

diff --git a/Engine/Plugins/Runtime/OpenXR/Source/OpenXRHMD/Private/OpenXRHMDModule.cpp b/Engine/Plugins/Runtime/OpenXR/Source/OpenXRHMD/Private/OpenXRHMDModule.cpp
--- a/Engine/Plugins/Runtime/OpenXR/Source/OpenXRHMD/Private/OpenXRHMDModule.cpp
+++ b/Engine/Plugins/Runtime/OpenXR/Source/OpenXRHMD/Private/OpenXRHMDModule.cpp
@@ -53,6 +53,11 @@ TSharedPtr< class IXRTrackingSystem, ESPMode::ThreadSafe > FOpenXRHMDModule::Cre
 	auto ARModule = FModuleManager::LoadModulePtr<IOpenXRARModule>("OpenXRAR");
 	auto ARSystem = ARModule->CreateARSystem();
 
+	if (!Instance)
+	{
+		return nullptr;
+	}
+
 	auto OpenXRHMD = FSceneViewExtensions::NewExtension<FOpenXRHMD>(Instance, RenderBridge, EnabledExtensions, ExtensionPlugins, ARSystem);
 	if (OpenXRHMD->IsInitialized())
 	{
