Face locked stereo layers emulation added.
Also fixed a bug for which, if SetupFrameQuadLayers_RenderThread was called more than once with GetStereoLayersDirty = true, NativeQuadLayers were recreated from scratch and lost the current state of bUpdateTexture. This led to static swapchains trying to acquire a second image making the app crash.

diff --git a/Engine/Plugins/Runtime/OpenXR/Source/OpenXRHMD/Private/OpenXRHMD.cpp b/Engine/Plugins/Runtime/OpenXR/Source/OpenXRHMD/Private/OpenXRHMD.cpp
--- a/Engine/Plugins/Runtime/OpenXR/Source/OpenXRHMD/Private/OpenXRHMD.cpp
+++ b/Engine/Plugins/Runtime/OpenXR/Source/OpenXRHMD/Private/OpenXRHMD.cpp
@@ -64,6 +64,12 @@ static TAutoConsoleVariable<int32> CVarOpenXREnvironmentBlendMode(
 	TEXT("Override the XrEnvironmentBlendMode used when submitting frames. 1 = Opaque, 2 = Additive, 3 = Alpha Blend\n"),
 	ECVF_Default);
 
+static TAutoConsoleVariable<bool> CVarOpenXRForceStereoLayerEmulation(
+	TEXT("xr.OpenXRForceStereoLayerEmulation"),
+	false,
+	TEXT("Force the emulation of stereo layers instead of using native ones (if supported).\n"),
+	ECVF_Default);
+
 namespace {
 	static TSet<XrEnvironmentBlendMode> SupportedBlendModes{ XR_ENVIRONMENT_BLEND_MODE_ALPHA_BLEND, XR_ENVIRONMENT_BLEND_MODE_ADDITIVE, XR_ENVIRONMENT_BLEND_MODE_OPAQUE };
 	static TSet<XrViewConfigurationType> SupportedViewConfigurations{ XR_VIEW_CONFIGURATION_TYPE_PRIMARY_STEREO, XR_VIEW_CONFIGURATION_TYPE_PRIMARY_QUAD_VARJO };
@@ -929,9 +935,15 @@ void FOpenXRHMD::SetFinalViewRect(FRHICommandListImmediate& RHICmdList, const in
 	XrSwapchainSubImage& DepthImage = PipelinedLayerStateRendering.DepthImages[ViewIndex];
 	if (bDepthExtensionSupported)
 	{
-		DepthImage.imageArrayIndex = bIsMobileMultiViewEnabled && ViewIndex < 2 ? ViewIndex : 0;
+		DepthImage.imageArrayIndex = ColorImage.imageArrayIndex;
 		DepthImage.imageRect = ColorImage.imageRect;
 	}
+	XrSwapchainSubImage& EmulationImage = PipelinedLayerStateRendering.EmulatedLayerState.EmulationImages[ViewIndex];
+	if(PipelinedLayerStateRendering.EmulatedLayerState.bIsFaceLockedLayerEmulationActive)
+	{
+		EmulationImage.imageArrayIndex = ColorImage.imageArrayIndex;
+		EmulationImage.imageRect = ColorImage.imageRect;
+	}
 }
 
 EStereoscopicPass FOpenXRHMD::GetViewPassForIndex(bool bStereoRequested, int32 ViewIndex) const
@@ -1066,9 +1078,11 @@ void FOpenXRHMD::BeginRenderViewFamily(FSceneViewFamily& InViewFamily)
 
 	PipelinedLayerStateRendering.ProjectionLayers.SetNum(ViewConfigCount);
 	PipelinedLayerStateRendering.DepthLayers.SetNum(ViewConfigCount);
+	PipelinedLayerStateRendering.EmulatedLayerState.CompositedProjectionLayers.SetNum(ViewConfigCount);
 
 	PipelinedLayerStateRendering.ColorImages.SetNum(PipelinedFrameStateRendering.ViewConfigs.Num());
 	PipelinedLayerStateRendering.DepthImages.SetNum(PipelinedFrameStateRendering.ViewConfigs.Num());
+	PipelinedLayerStateRendering.EmulatedLayerState.EmulationImages.SetNum(PipelinedFrameStateRendering.ViewConfigs.Num());
 
 	if (SpectatorScreenController)
 	{
@@ -1083,7 +1097,7 @@ void FOpenXRHMD::PreRenderView_RenderThread(FRDGBuilder& GraphBuilder, FSceneVie
 
 void FOpenXRHMD::PostRenderView_RenderThread(FRDGBuilder& GraphBuilder, FSceneView& InView)
 {
-	DrawEmulatedQuadLayers_RenderThread(GraphBuilder, InView);
+	DrawEmulatedLayers_RenderThread(GraphBuilder, InView);
 }
 
 void FOpenXRHMD::PreRenderViewFamily_RenderThread(FRDGBuilder& GraphBuilder, FSceneViewFamily& ViewFamily)
@@ -1112,12 +1126,17 @@ void FOpenXRHMD::PostRenderViewFamily_RenderThread(FRDGBuilder& GraphBuilder, FS
 		// Update SubImages with latest swapchain
 		XrSwapchainSubImage& ColorImage = PipelinedLayerStateRendering.ColorImages[ViewIndex];
 		XrSwapchainSubImage& DepthImage = PipelinedLayerStateRendering.DepthImages[ViewIndex];
+		XrSwapchainSubImage& EmulationImage = PipelinedLayerStateRendering.EmulatedLayerState.EmulationImages[ViewIndex];
 
 		ColorImage.swapchain = PipelinedLayerStateRendering.ColorSwapchain.IsValid() ? static_cast<FOpenXRSwapchain*>(PipelinedLayerStateRendering.ColorSwapchain.Get())->GetHandle() : XR_NULL_HANDLE;
 		if (bDepthExtensionSupported)
 		{
 			DepthImage.swapchain = PipelinedLayerStateRendering.DepthSwapchain.IsValid() ? static_cast<FOpenXRSwapchain*>(PipelinedLayerStateRendering.DepthSwapchain.Get())->GetHandle() : XR_NULL_HANDLE;
 		}
+		if(PipelinedLayerStateRendering.EmulatedLayerState.bIsFaceLockedLayerEmulationActive)
+		{
+			EmulationImage.swapchain = PipelinedLayerStateRendering.EmulatedLayerState.EmulationSwapchain.IsValid() ? static_cast<FOpenXRSwapchain*>(PipelinedLayerStateRendering.EmulatedLayerState.EmulationSwapchain.Get())->GetHandle() : XR_NULL_HANDLE;
+		}
 
 		if (IsViewManagedByPlugin(ViewIndex))
 		{
@@ -1150,6 +1169,14 @@ void FOpenXRHMD::PostRenderViewFamily_RenderThread(FRDGBuilder& GraphBuilder, FS
 
 			Projection.next = &DepthLayer;
 		}
+		if (PipelinedLayerStateRendering.EmulatedLayerState.bIsFaceLockedLayerEmulationActive)
+		{
+			XrCompositionLayerProjectionView& CompositedProjection = PipelinedLayerStateRendering.EmulatedLayerState.CompositedProjectionLayers[ViewIndex];
+
+			CompositedProjection.type = XR_TYPE_COMPOSITION_LAYER_PROJECTION_VIEW;
+			CompositedProjection.next = nullptr;
+			CompositedProjection.subImage = EmulationImage;
+		}
 
 		for (IOpenXRExtensionPlugin* Module : ExtensionPlugins)
 		{
@@ -1460,7 +1487,6 @@ void FOpenXRHMD::EnumerateViews(FPipelinedFrameState& PipelineState)
 			XrView& View = PipelineState.Views[ViewIndex];
 			View.type = XR_TYPE_VIEW;
 			View.next = nullptr;
-			// FIXME: should be recommendedFov
 			View.fov = ViewFov[ViewIndex].recommendedFov;
 			View.pose = ToXrPose(FTransform::Identity);
 		}
@@ -1639,16 +1665,14 @@ bool FOpenXRHMD::OnStereoStartup()
 
 	// Some runtimes aren't compliant with their number of layers supported.
 	// We support a fallback by emulating non-facelocked layers
-	bNativeWorldQuadLayerSupport = SystemProperties.graphicsProperties.maxLayerCount >= XR_MIN_COMPOSITION_LAYERS_SUPPORTED;
+	bLayerSupportOpenXRCompliant = SystemProperties.graphicsProperties.maxLayerCount >= XR_MIN_COMPOSITION_LAYERS_SUPPORTED; 
 
 	// Enumerate the viewport configurations
 	uint32 ConfigurationCount;
 	TArray<XrViewConfigurationType> ViewConfigTypes;
 	XR_ENSURE(xrEnumerateViewConfigurations(Instance, System, 0, &ConfigurationCount, nullptr));
-	ViewConfigTypes.SetNum(ConfigurationCount);
 	// Fill the initial array with valid enum types (this will fail in the validation layer otherwise).
-	for (auto& TypeIter : ViewConfigTypes)
-		TypeIter = XR_VIEW_CONFIGURATION_TYPE_PRIMARY_MONO;
+	ViewConfigTypes.Init(XR_VIEW_CONFIGURATION_TYPE_PRIMARY_MONO, ConfigurationCount);
 	XR_ENSURE(xrEnumerateViewConfigurations(Instance, System, ConfigurationCount, &ConfigurationCount, ViewConfigTypes.GetData()));
 
 	// Select the first view configuration returned by the runtime that is supported.
@@ -1910,10 +1934,12 @@ void FOpenXRHMD::DestroySession()
 		});
 
 		NativeQuadLayers.Reset();
+		EmulatedFaceLockedLayers.Reset();
 
 		PipelinedLayerStateRendering.ColorSwapchain.Reset();
 		PipelinedLayerStateRendering.DepthSwapchain.Reset();
 		PipelinedLayerStateRendering.QuadSwapchains.Reset();
+		PipelinedLayerStateRendering.EmulatedLayerState.EmulationSwapchain.Reset();
 
 		// TODO: Once we handle OnFinishRendering_RHIThread + StopSession interactions
 		// properly, we can release these shared pointers in that function, and use
@@ -1921,6 +1947,7 @@ void FOpenXRHMD::DestroySession()
 		PipelinedLayerStateRHI.ColorSwapchain.Reset();
 		PipelinedLayerStateRHI.DepthSwapchain.Reset();
 		PipelinedLayerStateRHI.QuadSwapchains.Reset();
+		PipelinedLayerStateRHI.EmulatedLayerState.EmulationSwapchain.Reset();
 
 		PipelinedFrameStateGame.TrackingSpace.Reset();
 		PipelinedFrameStateRendering.TrackingSpace.Reset();
@@ -2154,7 +2181,6 @@ bool FOpenXRHMD::AllocateRenderTargetTexture(uint32 Index, uint32 SizeX, uint32
 		UnifiedCreateFlags |= TexCreate_SRGB;
 	}
 
-
 	// Temporary workaround to swapchain formats - OpenXR doesn't support 10-bit sRGB swapchains, so prefer 8-bit sRGB instead.
 	if (Format == PF_A2B10G10R10 && !RenderBridge->Support10BitSwapchain())
 	{
@@ -2184,6 +2210,21 @@ bool FOpenXRHMD::AllocateRenderTargetTexture(uint32 Index, uint32 SizeX, uint32
 	LastRequestedColorSwapchainFormat = Format;
 	LastActualColorSwapchainFormat = ActualFormat;
 
+	if(IsEmulatingStereoLayers())
+	{
+		// Create flags for emulation swapchain
+		ETextureCreateFlags EmulationCreateFlags = TexCreate_Dynamic | TexCreate_ShaderResource | TexCreate_RenderTargetable;
+
+		FXRSwapChainPtr& EmulationSwapchain = PipelinedLayerStateRendering.EmulatedLayerState.EmulationSwapchain;
+		const FRHITexture2D* const EmulationSwapchainTexture = EmulationSwapchain == nullptr ? nullptr : EmulationSwapchain->GetTexture2DArray() ? EmulationSwapchain->GetTexture2DArray() : EmulationSwapchain->GetTexture2D();
+		if (EmulationSwapchain == nullptr || EmulationSwapchainTexture == nullptr || Format != LastRequestedColorSwapchainFormat || EmulationSwapchainTexture->GetSizeX() != SizeX || EmulationSwapchainTexture->GetSizeY() != SizeY)
+		{
+			ensureMsgf(NumSamples == 1, TEXT("OpenXR supports MSAA swapchains, but engine logic expects the swapchain target to be 1x."));
+
+			EmulationSwapchain = RenderBridge->CreateSwapchain(Session, Format, ActualFormat, SizeX, SizeY, bIsMobileMultiViewEnabled ? 2 : 1, NumMips, NumSamples, EmulationCreateFlags, FClearValueBinding::Transparent);
+		}
+	}
+
 	// TODO: Pass in known depth parameters (format + flags)? Do we know that at viewport setup time?
 	AllocateDepthTextureInternal(Index, SizeX, SizeY, NumSamples, OutTargetableTexture->GetDesc().ArraySize);
 
@@ -2308,25 +2349,79 @@ void CreateNativeLayerSwapchain(FOpenXRLayer& Layer, TRefCountPtr<FOpenXRRenderB
 	}
 }
 
+bool FOpenXRHMD::IsEmulatingStereoLayers()
+{
+	return !bLayerSupportOpenXRCompliant || CVarOpenXRForceStereoLayerEmulation.GetValueOnAnyThread();
+}
+
 void FOpenXRHMD::SetupFrameQuadLayers_RenderThread(FRHICommandListImmediate& RHICmdList)
 {
 	ensure(IsInRenderingThread());
 
 	if (GetStereoLayersDirty())
 	{
+		// When SetupFrameQuadLayers_RenderThread is called more than once with GetStereoLayersDirty = true,
+		// NativeQuadLayers is rebuilt from the internal array of stereo layers. However, the internal array is
+		// not updated after a static swapchain is created for a layer and it always retains bUpdateTexture = true
+		// which leads to the swapchain trying to acquire a second image and resulting in a crash.
+		// This serves as a workaround that updates the internal state of the layers mirroring the most recent state
+		// of the NativeQuadLayers.
+		TArray<FOpenXRLayer> NativeQuadLayersBackup = NativeQuadLayers;
 		// Go over the dirtied layers to bin them into either native or emulated
-		EmulatedSceneLayers.Reset();
+		BackgroundCompositedEmulatedLayers.Reset();
+		EmulatedFaceLockedLayers.Reset();
 		NativeQuadLayers.Reset();
 
 		ForEachLayer([&](uint32 /* unused */, FOpenXRLayer& Layer)
 		{
-			// We use native layer for facelocked and when world-locked has proper support
-			if ((Layer.Desc.PositionType == ELayerType::FaceLocked) || bNativeWorldQuadLayerSupport)
+			// OpenXR currently supports only Quad layers
+			if (!Layer.Desc.HasShape<FQuadLayer>())
+			{
+				return;
+			}
+			if (IsEmulatingStereoLayers())
+			{
+
+				if (!Layer.Desc.IsVisible())
+				{
+					return;
+				}
+				if(Layer.Desc.PositionType == ELayerType::FaceLocked)
+				{
+					// If we have at least one native layer, use it to render the 
+					// composited image of all the emulated face locked layers.
+					if (SystemProperties.graphicsProperties.maxLayerCount > 0)
+					{
+						EmulatedFaceLockedLayers.Add(Layer.Desc);
+					}
+					else
+					{
+						BackgroundCompositedEmulatedLayers.Add(Layer.Desc);
+					}
+				}else // Layer is not face locked
+				{
+					BackgroundCompositedEmulatedLayers.Add(Layer.Desc);
+				}
+			} // OpenXR compliant layer support (16 layers).
+			else 
 			{
-				if (Layer.Desc.IsVisible() && Layer.Desc.HasShape<FQuadLayer>())
+				if (Layer.Desc.IsVisible())
 				{
 					CreateNativeLayerSwapchain(Layer, RenderBridge, Session);
-					NativeQuadLayers.Add(Layer);
+					FOpenXRLayer& LastLayer = NativeQuadLayers.Add_GetRef(Layer);
+					FOpenXRLayer* FoundLayer = NativeQuadLayersBackup.FindByPredicate([LayerId = LastLayer.GetLayerId()](const FOpenXRLayer& Layer)
+					{
+						if (Layer.GetLayerId() == LayerId)
+						{
+							return true;
+						}
+						return false;
+					});
+					if (FoundLayer != nullptr)
+					{
+						LastLayer.RightEye.bUpdateTexture = FoundLayer->RightEye.bUpdateTexture;
+						LastLayer.LeftEye.bUpdateTexture = FoundLayer->LeftEye.bUpdateTexture;
+					}
 				}
 				else
 				{
@@ -2335,10 +2430,6 @@ void FOpenXRHMD::SetupFrameQuadLayers_RenderThread(FRHICommandListImmediate& RHI
 					Layer.LeftEye.Swapchain.Reset();
 				}
 			}
-			else if (Layer.Desc.IsVisible())
-			{
-				EmulatedSceneLayers.Add(Layer.Desc);
-			}
 		});
 
 		auto LayerCompare = [](const auto& A, const auto& B)
@@ -2359,8 +2450,11 @@ void FOpenXRHMD::SetupFrameQuadLayers_RenderThread(FRHICommandListImmediate& RHI
 			return false;
 		};
 
-		EmulatedSceneLayers.Sort(LayerCompare);
+		BackgroundCompositedEmulatedLayers.Sort(LayerCompare);
+		EmulatedFaceLockedLayers.Sort(LayerCompare);
 		NativeQuadLayers.Sort(LayerCompare);
+
+		PipelinedLayerStateRendering.EmulatedLayerState.bIsFaceLockedLayerEmulationActive = !EmulatedFaceLockedLayers.IsEmpty();
 	}
 
 	const FTransform InvTrackingToWorld = GetTrackingToWorldTransform().Inverse();
@@ -2426,62 +2520,53 @@ void FOpenXRHMD::SetupFrameQuadLayers_RenderThread(FRHICommandListImmediate& RHI
 	}
 }
 
-void FOpenXRHMD::DrawEmulatedQuadLayers_RenderThread(FRDGBuilder& GraphBuilder, const FSceneView& InView)
-{
+void FOpenXRHMD::DrawEmulatedLayers_RenderThread(FRDGBuilder& GraphBuilder, const FSceneView& InView)
+{	
 	check(IsInRenderingThread());
 
-	if (bNativeWorldQuadLayerSupport || !IStereoRendering::IsStereoEyeView(InView))
+	if (!IsEmulatingStereoLayers() || !IStereoRendering::IsStereoEyeView(InView))
 	{
 		return;
 	}
 
-	// Partially borrowed from FDefaultStereoLayers
-	AddPass(GraphBuilder, RDG_EVENT_NAME("OpenXREmulatedLayerRender"), [this, &InView](FRHICommandListImmediate& RHICmdList)
+	DrawBackgroundCompositedEmulatedLayers_RenderThread(GraphBuilder, InView);
+	DrawEmulatedFaceLockedLayers_RenderThread(GraphBuilder, InView);
+}
+void FOpenXRHMD::DrawEmulatedFaceLockedLayers_RenderThread(FRDGBuilder& GraphBuilder, const FSceneView& InView)
+{
+	if (!PipelinedLayerStateRendering.EmulatedLayerState.bIsFaceLockedLayerEmulationActive || !PipelinedLayerStateRendering.EmulatedLayerState.EmulationSwapchain)
 	{
-		FViewMatrices ModifiedViewMatrices = InView.ViewMatrices;
-		ModifiedViewMatrices.HackRemoveTemporalAAProjectionJitter();
-		const FMatrix& ProjectionMatrix = ModifiedViewMatrices.GetProjectionMatrix();
-		const FMatrix& ViewProjectionMatrix = ModifiedViewMatrices.GetViewProjectionMatrix();
+		return;
+	}
 
-		// Calculate a view matrix that only adjusts for eye position, ignoring head position, orientation and world position.
-		FVector EyeShift;
-		FQuat EyeOrientation;
-		GetRelativeEyePose(IXRTrackingSystem::HMDDeviceId, InView.StereoViewIndex, EyeOrientation, EyeShift);
+	AddPass(GraphBuilder, RDG_EVENT_NAME("OpenXREmulatedFaceLockedLayerRender"), [this, &InView](FRHICommandListImmediate& RHICmdList)
+	{
+		FTexture2DRHIRef RenderTarget = PipelinedLayerStateRendering.EmulatedLayerState.EmulationSwapchain->GetTextureRef();
 
-		FMatrix EyeMatrix = FTranslationMatrix(-EyeShift) * FInverseRotationMatrix(EyeOrientation.Rotator()) * FMatrix(
-			FPlane(0, 0, 1, 0),
-			FPlane(1, 0, 0, 0),
-			FPlane(0, 1, 0, 0),
-			FPlane(0, 0, 0, 1));
+		FDefaultStereoLayers_LayerRenderParams RenderParams;
+		FRHIRenderPassInfo RPInfo = SetupEmulatedLayersRenderPass(RHICmdList, InView, EmulatedFaceLockedLayers, RenderTarget, RenderParams);
 
-		FQuat HmdOrientation = FQuat::Identity;
-		FVector HmdLocation = FVector::ZeroVector;
-		GetCurrentPose(IXRTrackingSystem::HMDDeviceId, HmdOrientation, HmdLocation);
+		RHICmdList.BeginRenderPass(RPInfo, TEXT("EmulatedFaceLockedStereoLayerRender"));
+		RHICmdList.SetViewport((float)RenderParams.Viewport.Min.X, (float)RenderParams.Viewport.Min.Y, 0.0f, (float)RenderParams.Viewport.Max.X, (float)RenderParams.Viewport.Max.Y, 1.0f);
 
-		FMatrix TrackerMatrix = FTranslationMatrix(-HmdLocation) * FInverseRotationMatrix(HmdOrientation.Rotator()) * EyeMatrix;
+		DrawClearQuadAlpha(RHICmdList, 0);
 
-		FDefaultStereoLayers::FLayerRenderParams RenderParams{
-			InView.UnscaledViewRect, // Viewport
-			{
-				ViewProjectionMatrix,				// WorldLocked,
-				TrackerMatrix * ProjectionMatrix,	// TrackerLocked,
-				EyeMatrix * ProjectionMatrix		// FaceLocked
-			}
-		};
+		FDefaultStereoLayers::StereoLayerRender(RHICmdList, EmulatedFaceLockedLayers, RenderParams);
 
-		TArray<FRHITransitionInfo, TInlineAllocator<16>> Infos;
-		for (const FLayerDesc& Layer : EmulatedSceneLayers)
-		{
-			Infos.Add(FRHITransitionInfo(Layer.Texture, ERHIAccess::Unknown, ERHIAccess::SRVGraphics));
-		}
-		if (Infos.Num())
-		{
-			RHICmdList.Transition(Infos);
-		}
+		RHICmdList.EndRenderPass();
+	});
+}
 
+void FOpenXRHMD::DrawBackgroundCompositedEmulatedLayers_RenderThread(FRDGBuilder& GraphBuilder, const FSceneView& InView)
+{
+	// Partially borrowed from FDefaultStereoLayers
+	AddPass(GraphBuilder, RDG_EVENT_NAME("OpenXREmulatedLayerRender"), [this, &InView](FRHICommandListImmediate& RHICmdList)
+	{
 		FTexture2DRHIRef RenderTarget = InView.Family->RenderTarget->GetRenderTargetTexture();
 
-		FRHIRenderPassInfo RPInfo(RenderTarget, ERenderTargetActions::Load_Store);
+		FDefaultStereoLayers_LayerRenderParams RenderParams;
+		FRHIRenderPassInfo RPInfo = SetupEmulatedLayersRenderPass(RHICmdList, InView, BackgroundCompositedEmulatedLayers, RenderTarget, RenderParams);
+
 		RHICmdList.BeginRenderPass(RPInfo, TEXT("EmulatedStereoLayerRender"));
 		RHICmdList.SetViewport((float)RenderParams.Viewport.Min.X, (float)RenderParams.Viewport.Min.Y, 0.0f, (float)RenderParams.Viewport.Max.X, (float)RenderParams.Viewport.Max.Y, 1.0f);
 
@@ -2490,12 +2575,64 @@ void FOpenXRHMD::DrawEmulatedQuadLayers_RenderThread(FRDGBuilder& GraphBuilder,
 			DrawClearQuad(RHICmdList, FLinearColor::Black);
 		}
 
-		FDefaultStereoLayers::StereoLayerRender(RHICmdList, EmulatedSceneLayers, RenderParams);
+		FDefaultStereoLayers::StereoLayerRender(RHICmdList, BackgroundCompositedEmulatedLayers, RenderParams);
 
 		RHICmdList.EndRenderPass();
 	});
 }
 
+FRHIRenderPassInfo FOpenXRHMD::SetupEmulatedLayersRenderPass(FRHICommandListImmediate& RHICmdList, const FSceneView& InView, TArray<IStereoLayers::FLayerDesc>& Layers, FTexture2DRHIRef RenderTarget, FDefaultStereoLayers_LayerRenderParams& OutRenderParams)
+{
+	OutRenderParams = CalculateEmulatedLayerRenderParams(InView);
+	TArray<FRHITransitionInfo, TInlineAllocator<16>> Infos;
+	for (const FLayerDesc& Layer : Layers)
+	{
+		Infos.Add(FRHITransitionInfo(Layer.Texture, ERHIAccess::Unknown, ERHIAccess::SRVGraphics));
+	}
+	if (Infos.Num())
+	{
+		RHICmdList.Transition(Infos);
+	}
+
+	FRHIRenderPassInfo RPInfo(RenderTarget, ERenderTargetActions::Load_Store);
+	return RPInfo;
+}
+
+FDefaultStereoLayers_LayerRenderParams FOpenXRHMD::CalculateEmulatedLayerRenderParams(const FSceneView& InView)
+{
+	FViewMatrices ModifiedViewMatrices = InView.ViewMatrices;
+	ModifiedViewMatrices.HackRemoveTemporalAAProjectionJitter();
+	const FMatrix& ProjectionMatrix = ModifiedViewMatrices.GetProjectionMatrix();
+	const FMatrix& ViewProjectionMatrix = ModifiedViewMatrices.GetViewProjectionMatrix();
+
+	// Calculate a view matrix that only adjusts for eye position, ignoring head position, orientation and world position.
+	FVector EyeShift;
+	FQuat EyeOrientation;
+	GetRelativeEyePose(IXRTrackingSystem::HMDDeviceId, InView.StereoViewIndex, EyeOrientation, EyeShift);
+
+	FMatrix EyeMatrix = FTranslationMatrix(-EyeShift) * FInverseRotationMatrix(EyeOrientation.Rotator()) * FMatrix(
+		FPlane(0, 0, 1, 0),
+		FPlane(1, 0, 0, 0),
+		FPlane(0, 1, 0, 0),
+		FPlane(0, 0, 0, 1));
+
+	FQuat HmdOrientation = FQuat::Identity;
+	FVector HmdLocation = FVector::ZeroVector;
+	GetCurrentPose(IXRTrackingSystem::HMDDeviceId, HmdOrientation, HmdLocation);
+
+	FMatrix TrackerMatrix = FTranslationMatrix(-HmdLocation) * FInverseRotationMatrix(HmdOrientation.Rotator()) * EyeMatrix;
+
+	FDefaultStereoLayers_LayerRenderParams RenderParams{
+		InView.UnscaledViewRect, // Viewport
+		{
+			ViewProjectionMatrix,				// WorldLocked,
+			TrackerMatrix * ProjectionMatrix,	// TrackerLocked,
+			EyeMatrix * ProjectionMatrix		// FaceLocked
+		}
+	};
+	return RenderParams;
+}
+
 void FOpenXRHMD::OnBeginRendering_RenderThread(FRHICommandListImmediate& RHICmdList, FSceneViewFamily& ViewFamily)
 {
 	ensure(IsInRenderingThread());
@@ -2524,11 +2661,19 @@ void FOpenXRHMD::OnBeginRendering_RenderThread(FRHICommandListImmediate& RHICmdL
 
 		// Apply the base HMD pose to each eye pose, we will late update this pose for late update in another callback
 		FTransform BasePose(ViewFamily.Views[ViewIndex]->BaseHmdOrientation, ViewFamily.Views[ViewIndex]->BaseHmdLocation);
-		XrCompositionLayerProjectionView& Projection = PipelinedLayerStateRendering.ProjectionLayers[ViewIndex];
 		FTransform BasePoseTransform = EyePose * BasePose;
 		BasePoseTransform.NormalizeRotation();
+		
+		XrCompositionLayerProjectionView& Projection = PipelinedLayerStateRendering.ProjectionLayers[ViewIndex];
 		Projection.pose = ToXrPose(BasePoseTransform, WorldToMeters);
 		Projection.fov = View.fov;
+
+		if(PipelinedLayerStateRendering.EmulatedLayerState.bIsFaceLockedLayerEmulationActive)
+		{
+			XrCompositionLayerProjectionView& CompositedProjection = PipelinedLayerStateRendering.EmulatedLayerState.CompositedProjectionLayers[ViewIndex];
+			CompositedProjection.pose = ToXrPose(EyePose, WorldToMeters);
+			CompositedProjection.fov = View.fov;
+		}
 	}
 
 	SetupFrameQuadLayers_RenderThread(RHICmdList);
@@ -2559,9 +2704,13 @@ void FOpenXRHMD::OnBeginRendering_RenderThread(FRHICommandListImmediate& RHICmdL
 
 		FXRSwapChainPtr ColorSwapchain = PipelinedLayerStateRendering.ColorSwapchain;
 		FXRSwapChainPtr DepthSwapchain = PipelinedLayerStateRendering.DepthSwapchain;
-		RHICmdList.EnqueueLambda([this, FrameState = PipelinedFrameStateRendering, ColorSwapchain, DepthSwapchain](FRHICommandListImmediate& InRHICmdList)
+		// This swapchain might not be present depending on the platform support for stereo layers.
+		// Always check for sanity before using it.
+		FXRSwapChainPtr EmulationSwapchain = PipelinedLayerStateRendering.EmulatedLayerState.EmulationSwapchain;
+
+		RHICmdList.EnqueueLambda([this, FrameState = PipelinedFrameStateRendering, ColorSwapchain, DepthSwapchain, EmulationSwapchain](FRHICommandListImmediate& InRHICmdList)
 		{
-			OnBeginRendering_RHIThread(FrameState, ColorSwapchain, DepthSwapchain);
+			OnBeginRendering_RHIThread(FrameState, ColorSwapchain, DepthSwapchain, EmulationSwapchain);
 		});
 
 		// We need to sync with the RHI thread to ensure we've acquired the next swapchain image.
@@ -2624,11 +2773,19 @@ void FOpenXRHMD::OnLateUpdateApplied_RenderThread(FRHICommandListImmediate& RHIC
 
 		// Update the field-of-view to match the final projection matrix
 		Projection.fov = View.fov;
+
+		if(PipelinedLayerStateRendering.EmulatedLayerState.bIsFaceLockedLayerEmulationActive)
+		{
+			XrCompositionLayerProjectionView& CompositedProjection = PipelinedLayerStateRendering.EmulatedLayerState.CompositedProjectionLayers[ViewIndex];
+			CompositedProjection.pose = ToXrPose(EyePose, GetWorldToMetersScale());
+			CompositedProjection.fov = View.fov;
+		}
 	}
 
-	RHICmdList.EnqueueLambda([this, ProjectionLayers = PipelinedLayerStateRendering.ProjectionLayers](FRHICommandListImmediate& InRHICmdList)
+	RHICmdList.EnqueueLambda([this, ProjectionLayers = PipelinedLayerStateRendering.ProjectionLayers, CompositedProjectionLayers = PipelinedLayerStateRendering.EmulatedLayerState.CompositedProjectionLayers](FRHICommandListImmediate& InRHICmdList)
 	{
 		PipelinedLayerStateRHI.ProjectionLayers = ProjectionLayers;
+		PipelinedLayerStateRHI.EmulatedLayerState.CompositedProjectionLayers = CompositedProjectionLayers;
 	});
 }
 
@@ -2649,7 +2806,7 @@ void FOpenXRHMD::OnBeginRendering_GameThread()
 
 			// If we are emulating layers, we still need to submit background layer since we composite into it
 			PipelinedLayerStateRendering.bBackgroundLayerVisible = bBackgroundLayerVisible;
-			PipelinedLayerStateRendering.bSubmitBackgroundLayer = bBackgroundLayerVisible || !bNativeWorldQuadLayerSupport;
+			PipelinedLayerStateRendering.bSubmitBackgroundLayer = bBackgroundLayerVisible || IsEmulatingStereoLayers();
 		});
 }
 
@@ -2894,7 +3051,7 @@ void FOpenXRHMD::RequestExitApp()
 	}
 }
 
-void FOpenXRHMD::OnBeginRendering_RHIThread(const FPipelinedFrameState& InFrameState, FXRSwapChainPtr ColorSwapchain, FXRSwapChainPtr DepthSwapchain)
+void FOpenXRHMD::OnBeginRendering_RHIThread(const FPipelinedFrameState& InFrameState, FXRSwapChainPtr ColorSwapchain, FXRSwapChainPtr DepthSwapchain, FXRSwapChainPtr EmulationSwapchain)
 {
 	ensure(IsInRenderingThread() || IsInRHIThread());
 
@@ -2929,6 +3086,7 @@ void FOpenXRHMD::OnBeginRendering_RHIThread(const FPipelinedFrameState& InFrameS
 		// TODO Possibly move these Waits to SetFinalViewRect??
 		PipelinedLayerStateRHI.ColorSwapchain = ColorSwapchain;
 		PipelinedLayerStateRHI.DepthSwapchain = DepthSwapchain;
+		PipelinedLayerStateRHI.EmulatedLayerState.EmulationSwapchain = EmulationSwapchain;
 
 		// We need a new swapchain image unless we've already acquired one for rendering
 		if (!bIsRendering && ColorSwapchain)
@@ -2940,6 +3098,11 @@ void FOpenXRHMD::OnBeginRendering_RHIThread(const FPipelinedFrameState& InFrameS
 				DepthSwapchain->IncrementSwapChainIndex_RHIThread();
 				DepthSwapchain->WaitCurrentImage_RHIThread(OPENXR_SWAPCHAIN_WAIT_TIMEOUT);
 			}
+			if(PipelinedLayerStateRHI.EmulatedLayerState.bIsFaceLockedLayerEmulationActive && EmulationSwapchain)
+			{
+				EmulationSwapchain->IncrementSwapChainIndex_RHIThread();
+				EmulationSwapchain->WaitCurrentImage_RHIThread(OPENXR_SWAPCHAIN_WAIT_TIMEOUT);
+			}
 		}
 
 		bIsRendering = true;
@@ -2977,6 +3140,10 @@ void FOpenXRHMD::OnFinishRendering_RHIThread()
 		{
 			PipelinedLayerStateRHI.DepthSwapchain->ReleaseCurrentImage_RHIThread();
 		}
+		if(PipelinedLayerStateRHI.EmulatedLayerState.bIsFaceLockedLayerEmulationActive && PipelinedLayerStateRHI.EmulatedLayerState.EmulationSwapchain)
+		{
+			PipelinedLayerStateRHI.EmulatedLayerState.EmulationSwapchain->ReleaseCurrentImage_RHIThread();
+		}
 	}
 
 	FReadScopeLock Lock(SessionHandleMutex);
@@ -3004,6 +3171,22 @@ void FOpenXRHMD::OnFinishRendering_RHIThread()
 			Layer.layerFlags |= XR_COMPOSITION_LAYER_UNPREMULTIPLIED_ALPHA_BIT;
 #endif
 		}
+		
+		XrCompositionLayerProjection CompositedLayer = {};
+		if (PipelinedLayerStateRHI.EmulatedLayerState.bIsFaceLockedLayerEmulationActive && PipelinedLayerStateRHI.EmulatedLayerState.EmulationSwapchain)
+		{
+			CompositedLayer.type = XR_TYPE_COMPOSITION_LAYER_PROJECTION;
+			CompositedLayer.next = nullptr;
+			// Alpha always enabled to allow for transparency between the composited layers.
+			CompositedLayer.layerFlags = XR_COMPOSITION_LAYER_BLEND_TEXTURE_SOURCE_ALPHA_BIT;
+			{
+				FReadScopeLock DeviceLock(DeviceMutex);
+				CompositedLayer.space = DeviceSpaces[HMDDeviceId].Space;
+			}
+			CompositedLayer.viewCount = PipelinedLayerStateRHI.EmulatedLayerState.CompositedProjectionLayers.Num();
+			CompositedLayer.views = PipelinedLayerStateRHI.EmulatedLayerState.CompositedProjectionLayers.GetData();
+			Headers.Add(reinterpret_cast<const XrCompositionLayerBaseHeader*>(&CompositedLayer));
+		}
 
 		for (const XrCompositionLayerQuad& Quad : PipelinedLayerStateRHI.QuadLayers)
 		{
diff --git a/Engine/Plugins/Runtime/OpenXR/Source/OpenXRHMD/Private/OpenXRHMD.h b/Engine/Plugins/Runtime/OpenXR/Source/OpenXRHMD/Private/OpenXRHMD.h
--- a/Engine/Plugins/Runtime/OpenXR/Source/OpenXRHMD/Private/OpenXRHMD.h
+++ b/Engine/Plugins/Runtime/OpenXR/Source/OpenXRHMD/Private/OpenXRHMD.h
@@ -24,6 +24,7 @@ class FSceneViewFamily;
 class UCanvas;
 class FOpenXRRenderBridge;
 class IOpenXRInputModule;
+struct FDefaultStereoLayers_LayerRenderParams;
 
 /**
  * Simple Head Mounted Display
@@ -91,6 +92,17 @@ public:
 		bool bXrFrameStateUpdated = false;
 	};
 
+	struct FEmulatedLayerState
+	{
+		// These layers are used as a target to composite all the emulated face locked layers into
+		// and be sent to the compositor with VIEW tracking space to avoid reprojection.
+		TArray<XrCompositionLayerProjectionView> CompositedProjectionLayers;
+		TArray<XrSwapchainSubImage> EmulationImages;
+		// This swapchain is where the emulated face locked layers are rendered into.
+		FXRSwapChainPtr EmulationSwapchain;
+		bool bIsFaceLockedLayerEmulationActive = false;
+	};
+
 	struct FPipelinedLayerState
 	{
 		TArray<XrCompositionLayerQuad> QuadLayers;
@@ -104,6 +116,8 @@ public:
 		FXRSwapChainPtr DepthSwapchain;
 		TArray<FXRSwapChainPtr> QuadSwapchains;
 
+		FEmulatedLayerState EmulatedLayerState;
+
 		bool bBackgroundLayerVisible = true;
 		bool bSubmitBackgroundLayer = true;
 	};
@@ -249,7 +263,9 @@ protected:
 	void VRHeadsetRecenterDelegate();
 
 	void SetupFrameQuadLayers_RenderThread(FRHICommandListImmediate& RHICmdList);
-	void DrawEmulatedQuadLayers_RenderThread(FRDGBuilder& GraphBuilder, const FSceneView& InView);
+	void DrawEmulatedLayers_RenderThread(FRDGBuilder& GraphBuilder, const FSceneView& InView);
+	void DrawBackgroundCompositedEmulatedLayers_RenderThread(FRDGBuilder& GraphBuilder, const FSceneView& InView);
+	void DrawEmulatedFaceLockedLayers_RenderThread(FRDGBuilder& GraphBuilder, const FSceneView& InView);
 
 	/** TStereoLayerManager<FOpenXRLayer> */
 	void UpdateLayer(FOpenXRLayer& ManagerLayer, uint32 LayerId, bool bIsValid) override;
@@ -350,7 +366,7 @@ public:
 	virtual ~FOpenXRHMD();
 
 	void OnBeginSimulation_GameThread();
-	void OnBeginRendering_RHIThread(const FPipelinedFrameState& InFrameState, FXRSwapChainPtr ColorSwapchain, FXRSwapChainPtr DepthSwapchain);
+	void OnBeginRendering_RHIThread(const FPipelinedFrameState& InFrameState, FXRSwapChainPtr ColorSwapchain, FXRSwapChainPtr DepthSwapchain, FXRSwapChainPtr EmulationSwapchain);
 	void OnFinishRendering_RHIThread();
 
 	/** @return	True if the HMD was initialized OK */
@@ -374,6 +390,9 @@ public:
 
 private:
 	TArray<XrEnvironmentBlendMode> RetrieveEnvironmentBlendModes() const;
+	FDefaultStereoLayers_LayerRenderParams CalculateEmulatedLayerRenderParams(const FSceneView& InView);
+	FRHIRenderPassInfo SetupEmulatedLayersRenderPass(FRHICommandListImmediate& RHICmdList, const FSceneView& InView, TArray<IStereoLayers::FLayerDesc>& Layers, FTexture2DRHIRef RenderTarget, FDefaultStereoLayers_LayerRenderParams& OutRenderParams);
+	bool IsEmulatingStereoLayers();
 
 	bool					bStereoEnabled;
 	TAtomic<bool>			bIsRunning;
@@ -438,7 +457,8 @@ private:
 	FQuat					BaseOrientation;
 	FVector					BasePosition;
 
-	bool					bNativeWorldQuadLayerSupport;
-	TArray<IStereoLayers::FLayerDesc> EmulatedSceneLayers;
+	bool					bLayerSupportOpenXRCompliant;
+	TArray<IStereoLayers::FLayerDesc> BackgroundCompositedEmulatedLayers;
+	TArray<IStereoLayers::FLayerDesc> EmulatedFaceLockedLayers;
 	TArray<FOpenXRLayer>			  NativeQuadLayers;
 };
diff --git a/Engine/Source/Runtime/HeadMountedDisplay/Private/DefaultStereoLayers.cpp b/Engine/Source/Runtime/HeadMountedDisplay/Private/DefaultStereoLayers.cpp
--- a/Engine/Source/Runtime/HeadMountedDisplay/Private/DefaultStereoLayers.cpp
+++ b/Engine/Source/Runtime/HeadMountedDisplay/Private/DefaultStereoLayers.cpp
@@ -58,7 +58,7 @@ FDefaultStereoLayers::FDefaultStereoLayers(const FAutoRegister& AutoRegister, FH
 //=============================================================================
 
 // static
-void FDefaultStereoLayers::StereoLayerRender(FRHICommandListImmediate& RHICmdList, const TArray<FLayerDesc>& LayersToRender, const FLayerRenderParams& RenderParams)
+void FDefaultStereoLayers::StereoLayerRender(FRHICommandListImmediate& RHICmdList, const TArray<FLayerDesc>& LayersToRender, const FDefaultStereoLayers_LayerRenderParams& RenderParams)
 {
 	check(IsInRenderingThread());
 	if (!LayersToRender.Num())
@@ -235,7 +235,7 @@ void FDefaultStereoLayers::PostRenderView_RenderThread(FRDGBuilder& GraphBuilder
 		FVector HmdLocation = HmdTransform.GetTranslation();
 		FMatrix TrackerMatrix = FTranslationMatrix(-HmdLocation) * FInverseRotationMatrix(HmdOrientation.Rotator()) * EyeMatrix;
 
-		FLayerRenderParams RenderParams{
+		FDefaultStereoLayers_LayerRenderParams RenderParams{
 			InView.UnscaledViewRect, // Viewport
 			{
 				ViewProjectionMatrix,				// WorldLocked,
diff --git a/Engine/Source/Runtime/HeadMountedDisplay/Public/DefaultStereoLayers.h b/Engine/Source/Runtime/HeadMountedDisplay/Public/DefaultStereoLayers.h
--- a/Engine/Source/Runtime/HeadMountedDisplay/Public/DefaultStereoLayers.h
+++ b/Engine/Source/Runtime/HeadMountedDisplay/Public/DefaultStereoLayers.h
@@ -5,6 +5,13 @@
 #include "SceneViewExtension.h"
 class FHeadMountedDisplayBase;
 
+/** Experimental struct */
+struct HEADMOUNTEDDISPLAY_API FDefaultStereoLayers_LayerRenderParams
+{
+	FIntRect Viewport;
+	FMatrix RenderMatrices[3];
+};
+
 /** 
  *	Default implementation of stereo layers for platforms that require emulating layer support.
  *
@@ -29,15 +36,8 @@ public:
 		return false;
 	}
 
-	/** Experimental struct */
-	struct FLayerRenderParams
-	{
-		FIntRect Viewport;
-		FMatrix RenderMatrices[3];
-	};
-
 	/** Experimental method */
-	static void StereoLayerRender(FRHICommandListImmediate& RHICmdList, const TArray<FLayerDesc>& LayersToRender, const FLayerRenderParams& RenderParams);
+	static void StereoLayerRender(FRHICommandListImmediate& RHICmdList, const TArray<FLayerDesc>& LayersToRender, const FDefaultStereoLayers_LayerRenderParams& RenderParams);
 
 protected:
 	
@@ -56,4 +56,4 @@ protected:
 	TArray<FLayerDesc> SortedOverlayLayers;
 
 	friend class FHeadMountedDisplayBase;
-};
\ No newline at end of file
+};
