Fixed a bug for which the first image of the emulation swapchain was not acquired.

diff --git a/Engine/Plugins/Runtime/OpenXR/Source/OpenXRHMD/Private/OpenXRHMD.cpp b/Engine/Plugins/Runtime/OpenXR/Source/OpenXRHMD/Private/OpenXRHMD.cpp
--- a/Engine/Plugins/Runtime/OpenXR/Source/OpenXRHMD/Private/OpenXRHMD.cpp
+++ b/Engine/Plugins/Runtime/OpenXR/Source/OpenXRHMD/Private/OpenXRHMD.cpp
@@ -2708,8 +2708,10 @@ void FOpenXRHMD::OnBeginRendering_RenderThread(FRHICommandListImmediate& RHICmdL
 		// Always check for sanity before using it.
 		FXRSwapChainPtr EmulationSwapchain = PipelinedLayerStateRendering.EmulatedLayerState.EmulationSwapchain;
 
-		RHICmdList.EnqueueLambda([this, FrameState = PipelinedFrameStateRendering, ColorSwapchain, DepthSwapchain, EmulationSwapchain](FRHICommandListImmediate& InRHICmdList)
+		RHICmdList.EnqueueLambda([this, FrameState = PipelinedFrameStateRendering, bIsFaceLockedLayerEmulationActive = PipelinedLayerStateRendering.EmulatedLayerState.bIsFaceLockedLayerEmulationActive, ColorSwapchain, DepthSwapchain, EmulationSwapchain](FRHICommandListImmediate& InRHICmdList)
 		{
+			// This needs to be passed here or the first frame will not have the correct value and will not acquire the swapchain image even if the emulation is active.
+			PipelinedLayerStateRHI.EmulatedLayerState.bIsFaceLockedLayerEmulationActive = bIsFaceLockedLayerEmulationActive;
 			OnBeginRendering_RHIThread(FrameState, ColorSwapchain, DepthSwapchain, EmulationSwapchain);
 		});
 
