Improve handling of VRNotificationsComponent delegates.

Also adds a new event to allow devs to respond to interaction profile changes.

diff --git a/Engine/Plugins/Runtime/OpenXR/Source/OpenXRHMD/Private/OpenXRHMD.cpp b/Engine/Plugins/Runtime/OpenXR/Source/OpenXRHMD/Private/OpenXRHMD.cpp
--- a/Engine/Plugins/Runtime/OpenXR/Source/OpenXRHMD/Private/OpenXRHMD.cpp
+++ b/Engine/Plugins/Runtime/OpenXR/Source/OpenXRHMD/Private/OpenXRHMD.cpp
@@ -735,6 +735,7 @@ void FOpenXRHMD::Recenter(EOrientPositionSelector::Type Selector, float Yaw)
 	}
 
 	bTrackingSpaceInvalid = true;
+	OnTrackingOriginChanged();
 }
 
 void FOpenXRHMD::SetBaseRotation(const FRotator& InBaseRotation)
@@ -1234,7 +1235,7 @@ FOpenXRHMD::FOpenXRHMD(const FAutoRegister& AutoRegister, XrInstance InInstance,
 	, InputModule(nullptr)
 	, ExtensionPlugins(std::move(InExtensionPlugins))
 	, Instance(InInstance)
-	, System(XR_NULL_SYSTEM_ID)
+	, System(IOpenXRHMDModule::Get().GetSystemId())
 	, Session(XR_NULL_HANDLE)
 	, LocalSpace(XR_NULL_HANDLE)
 	, StageSpace(XR_NULL_HANDLE)
@@ -2907,6 +2908,15 @@ bool FOpenXRHMD::OnStartGameFrame(FWorldContext& WorldContext)
 
 	RefreshTrackingToWorldTransform(WorldContext);
 
+	if (!System)
+	{
+		System = IOpenXRHMDModule::Get().GetSystemId();
+		if (System)
+		{
+			FCoreDelegates::VRHeadsetReconnected.Broadcast();
+		}
+	}
+
 	// Process all pending messages.
 	XrEventDataBuffer event;
 	while (ReadNextEvent(&event))
@@ -2928,6 +2938,7 @@ bool FOpenXRHMD::OnStartGameFrame(FWorldContext& WorldContext)
 				{
 					GEngine->SetMaxFPS(0);
 				}
+				FCoreDelegates::VRHeadsetPutOnHead.Broadcast();
 				bIsReady = true;
 				StartSession();
 			}
@@ -2945,6 +2956,7 @@ bool FOpenXRHMD::OnStartGameFrame(FWorldContext& WorldContext)
 				{
 					GEngine->SetMaxFPS(OPENXR_PAUSED_IDLE_FPS);
 				}
+				FCoreDelegates::VRHeadsetRemovedFromHead.Broadcast();
 				bIsReady = false;
 				StopSession();
 			}
@@ -2955,6 +2967,12 @@ bool FOpenXRHMD::OnStartGameFrame(FWorldContext& WorldContext)
 				{
 					GEngine->SetMaxFPS(0);
 				}
+
+				if (SessionState.state == XR_SESSION_STATE_LOSS_PENDING)
+				{
+					FCoreDelegates::VRHeadsetLost.Broadcast();
+					System = XR_NULL_SYSTEM_ID;
+				}
 				
 				FApp::SetHasVRFocus(false);
 
@@ -3010,10 +3028,16 @@ bool FOpenXRHMD::OnStartGameFrame(FWorldContext& WorldContext)
 
 			break;
 		}
+		case XR_TYPE_EVENT_DATA_INTERACTION_PROFILE_CHANGED:
+		{
+			OnInteractionProfileChanged();
+			break;
+		}
 		case XR_TYPE_EVENT_DATA_VISIBILITY_MASK_CHANGED_KHR:
 		{
 			bHiddenAreaMaskSupported = ensure(IsExtensionEnabled(XR_KHR_VISIBILITY_MASK_EXTENSION_NAME));  // Ensure fail indicates a non-conformant openxr implementation.
 			bNeedReBuildOcclusionMesh = true;
+			break;
 		}
 		}
 	}
diff --git a/Engine/Source/Runtime/HeadMountedDisplay/Private/VRNotificationsComponent.cpp b/Engine/Source/Runtime/HeadMountedDisplay/Private/VRNotificationsComponent.cpp
--- a/Engine/Source/Runtime/HeadMountedDisplay/Private/VRNotificationsComponent.cpp
+++ b/Engine/Source/Runtime/HeadMountedDisplay/Private/VRNotificationsComponent.cpp
@@ -27,6 +27,7 @@ void UVRNotificationsComponent::OnRegister()
 	FCoreDelegates::VRControllerRecentered.AddUObject(this, &UVRNotificationsComponent::VRControllerRecentered_Handler);
 	FXRTrackingSystemDelegates::OnXRTrackingOriginChanged.AddUObject(this, &UVRNotificationsComponent::XRTrackingOriginChanged_Handler);
 	FXRTrackingSystemDelegates::OnXRPlayAreaChanged.AddUObject(this, &UVRNotificationsComponent::XRPlayAreaChanged_Handler);
+	FXRTrackingSystemDelegates::OnXRInteractionProfileChanged.AddUObject(this, &UVRNotificationsComponent::XRInteractionProfileChanged_Handler);
 }
 
 void UVRNotificationsComponent::OnUnregister()
@@ -44,5 +45,6 @@ void UVRNotificationsComponent::OnUnregister()
 	FCoreDelegates::VRControllerRecentered.RemoveAll(this);
 	FXRTrackingSystemDelegates::OnXRTrackingOriginChanged.RemoveAll(this);
 	FXRTrackingSystemDelegates::OnXRPlayAreaChanged.RemoveAll(this);
+	FXRTrackingSystemDelegates::OnXRInteractionProfileChanged.RemoveAll(this);
 }
 
diff --git a/Engine/Source/Runtime/HeadMountedDisplay/Private/XRTrackingSystemBase.cpp b/Engine/Source/Runtime/HeadMountedDisplay/Private/XRTrackingSystemBase.cpp
--- a/Engine/Source/Runtime/HeadMountedDisplay/Private/XRTrackingSystemBase.cpp
+++ b/Engine/Source/Runtime/HeadMountedDisplay/Private/XRTrackingSystemBase.cpp
@@ -14,7 +14,7 @@
 // Tracking system delegates
 FXRTrackingSystemDelegates::FXRTrackingOriginChanged FXRTrackingSystemDelegates::OnXRTrackingOriginChanged;
 FXRTrackingSystemDelegates::FXRPlayAreaChanged FXRTrackingSystemDelegates::OnXRPlayAreaChanged;
-
+FXRTrackingSystemDelegates::FXRPlayAreaChanged FXRTrackingSystemDelegates::OnXRInteractionProfileChanged;
 
 
 
diff --git a/Engine/Source/Runtime/HeadMountedDisplay/Public/VRNotificationsComponent.h b/Engine/Source/Runtime/HeadMountedDisplay/Public/VRNotificationsComponent.h
--- a/Engine/Source/Runtime/HeadMountedDisplay/Public/VRNotificationsComponent.h
+++ b/Engine/Source/Runtime/HeadMountedDisplay/Public/VRNotificationsComponent.h
@@ -64,14 +64,18 @@ class HEADMOUNTEDDISPLAY_API UVRNotificationsComponent : public UActorComponent
 	UPROPERTY(BlueprintAssignable)
 	FVRNotificationsDelegate VRControllerRecenteredDelegate;
 
-	// This will be called when the VR system recenters a controller.  
+	// This will be called whenever the tracking origin is altered.
 	UPROPERTY(BlueprintAssignable)
 	FVRNotificationsDelegate XRTrackingOriginChangedDelegate;
 
-	// This will be called when the VR system recenters a controller.  
+	// This will be called whenever the play area is altered.
 	UPROPERTY(BlueprintAssignable)
 	FVRNotificationsDelegate XRPlayAreaChangedDelegate;
 
+	// This will be called whenever the interaction profile is altered. 
+	UPROPERTY(BlueprintAssignable)
+	FVRNotificationsDelegate XRInteractionProfileChangedDelegate;
+
 public:
 	void OnRegister() override;
 	void OnUnregister() override;
@@ -89,6 +93,7 @@ private:
 	void VRControllerRecentered_Handler()		{ VRControllerRecenteredDelegate.Broadcast(); }
 	void XRTrackingOriginChanged_Handler(const IXRTrackingSystem* TrackingSys)		{ XRTrackingOriginChangedDelegate.Broadcast(); }
 	void XRPlayAreaChanged_Handler()			{ XRPlayAreaChangedDelegate.Broadcast(); }
+	void XRInteractionProfileChanged_Handler()	{ XRPlayAreaChangedDelegate.Broadcast(); }
 };
 
 
diff --git a/Engine/Source/Runtime/HeadMountedDisplay/Public/XRTrackingSystemBase.h b/Engine/Source/Runtime/HeadMountedDisplay/Public/XRTrackingSystemBase.h
--- a/Engine/Source/Runtime/HeadMountedDisplay/Public/XRTrackingSystemBase.h
+++ b/Engine/Source/Runtime/HeadMountedDisplay/Public/XRTrackingSystemBase.h
@@ -14,6 +14,9 @@ public:
 
 	DECLARE_MULTICAST_DELEGATE(FXRPlayAreaChanged);
 	static FXRPlayAreaChanged OnXRPlayAreaChanged;
+
+	DECLARE_MULTICAST_DELEGATE(FXRInteractionProfileChanged);
+	static FXRPlayAreaChanged OnXRInteractionProfileChanged;
 };
 
 /** 
@@ -190,6 +193,14 @@ protected:
 		FXRTrackingSystemDelegates::OnXRPlayAreaChanged.Broadcast();
 	}
 
+	/**
+	 * Meant to be called by sub-classes whenever the interaction profile is altered.
+	 */
+	virtual void OnInteractionProfileChanged()
+	{
+		FXRTrackingSystemDelegates::OnXRInteractionProfileChanged.Broadcast();
+	}
+
 	/**
 	 * Computes the project's tracking-to-world transform based off how the user 
 	 * has set up their camera system (assumes the camera is parented to the XR 
