OpenXR: emulated layers swapchain should use prescribed layer format

diff --git a/Engine/Plugins/Runtime/OpenXR/Source/OpenXRHMD/Private/OpenXRHMD.cpp b/Engine/Plugins/Runtime/OpenXR/Source/OpenXRHMD/Private/OpenXRHMD.cpp
--- a/Engine/Plugins/Runtime/OpenXR/Source/OpenXRHMD/Private/OpenXRHMD.cpp
+++ b/Engine/Plugins/Runtime/OpenXR/Source/OpenXRHMD/Private/OpenXRHMD.cpp
@@ -2213,16 +2213,14 @@ bool FOpenXRHMD::AllocateRenderTargetTexture(uint32 Index, uint32 SizeX, uint32
 
 	if(IsEmulatingStereoLayers())
 	{
-		// Create flags for emulation swapchain
-		ETextureCreateFlags EmulationCreateFlags = TexCreate_Dynamic | TexCreate_ShaderResource | TexCreate_RenderTargetable;
-
 		FXRSwapChainPtr& EmulationSwapchain = PipelinedLayerStateRendering.EmulatedLayerState.EmulationSwapchain;
 		const FRHITexture2D* const EmulationSwapchainTexture = EmulationSwapchain == nullptr ? nullptr : EmulationSwapchain->GetTexture2DArray() ? EmulationSwapchain->GetTexture2DArray() : EmulationSwapchain->GetTexture2D();
-		if (EmulationSwapchain == nullptr || EmulationSwapchainTexture == nullptr || Format != LastRequestedColorSwapchainFormat || EmulationSwapchainTexture->GetSizeX() != SizeX || EmulationSwapchainTexture->GetSizeY() != SizeY)
+		if (EmulationSwapchain == nullptr || EmulationSwapchainTexture == nullptr || EmulationSwapchainTexture->GetSizeX() != SizeX || EmulationSwapchainTexture->GetSizeY() != SizeY)
 		{
-			ensureMsgf(NumSamples == 1, TEXT("OpenXR supports MSAA swapchains, but engine logic expects the swapchain target to be 1x."));
+			const ETextureCreateFlags EmulationCreateFlags = TexCreate_Dynamic | TexCreate_ShaderResource | TexCreate_RenderTargetable;
 
-			EmulationSwapchain = RenderBridge->CreateSwapchain(Session, Format, ActualFormat, SizeX, SizeY, bIsMobileMultiViewEnabled ? 2 : 1, NumMips, NumSamples, EmulationCreateFlags, FClearValueBinding::Transparent);
+			uint8 UnusedActualFormat = 0;
+			EmulationSwapchain = RenderBridge->CreateSwapchain(Session, IStereoRenderTargetManager::GetStereoLayerPixelFormat(), UnusedActualFormat, SizeX, SizeY, bIsMobileMultiViewEnabled ? 2 : 1, NumMips, NumSamples, EmulationCreateFlags, FClearValueBinding::Transparent);
 		}
 	}
 
