OpenXR: quad layers respect component scale

diff --git a/Engine/Plugins/Runtime/OpenXR/Source/OpenXRHMD/Private/OpenXRHMD.cpp b/Engine/Plugins/Runtime/OpenXR/Source/OpenXRHMD/Private/OpenXRHMD.cpp
--- a/Engine/Plugins/Runtime/OpenXR/Source/OpenXRHMD/Private/OpenXRHMD.cpp
+++ b/Engine/Plugins/Runtime/OpenXR/Source/OpenXRHMD/Private/OpenXRHMD.cpp
@@ -2480,6 +2480,9 @@ void FOpenXRHMD::SetupFrameQuadLayers_RenderThread(FRHICommandListImmediate& RHI
 			Quad.subImage.imageArrayIndex = 0;
 			Quad.pose = ToXrPose(Layer.Desc.Transform * PositionTransform, WorldToMeters);
 
+			// The layer pose doesn't take the transform scale into consideration, so we need to manually apply it the quad size.
+			const FVector2D LayerComponentScaler(Layer.Desc.Transform.GetScale3D().Y, Layer.Desc.Transform.GetScale3D().Z);
+
 			// We need to copy each layer into an OpenXR swapchain so they can be displayed by the compositor
 			if (Layer.RightEye.Swapchain.IsValid() && Layer.Desc.Texture.IsValid())
 			{
@@ -2494,7 +2497,7 @@ void FOpenXRHMD::SetupFrameQuadLayers_RenderThread(FRHICommandListImmediate& RHI
 
 				Quad.subImage.imageRect = ToXrRect(Layer.GetRightViewportSize());
 				Quad.subImage.swapchain = static_cast<FOpenXRSwapchain*>(Layer.RightEye.Swapchain.Get())->GetHandle();
-				Quad.size = ToXrExtent2D(Layer.GetRightQuadSize(), WorldToMeters);
+				Quad.size = ToXrExtent2D(Layer.GetRightQuadSize() * LayerComponentScaler, WorldToMeters);
 				PipelinedLayerStateRendering.QuadLayers.Add(Quad);
 				PipelinedLayerStateRendering.QuadSwapchains.Add(Layer.RightEye.Swapchain);
 			}
@@ -2511,7 +2514,7 @@ void FOpenXRHMD::SetupFrameQuadLayers_RenderThread(FRHICommandListImmediate& RHI
 				Quad.eyeVisibility = XR_EYE_VISIBILITY_LEFT;
 				Quad.subImage.imageRect = ToXrRect(Layer.GetLeftViewportSize());
 				Quad.subImage.swapchain = static_cast<FOpenXRSwapchain*>(Layer.LeftEye.Swapchain.Get())->GetHandle();
-				Quad.size = ToXrExtent2D(Layer.GetLeftQuadSize(), WorldToMeters);
+				Quad.size = ToXrExtent2D(Layer.GetLeftQuadSize() * LayerComponentScaler, WorldToMeters);
 				PipelinedLayerStateRendering.QuadLayers.Add(Quad);
 				PipelinedLayerStateRendering.QuadSwapchains.Add(Layer.LeftEye.Swapchain);
 			}
