Fix crash when starting standalone game with OpenXR enabled, but no HMD is connected.

diff --git a/Engine/Plugins/Runtime/OpenXR/Source/OpenXRHMD/Private/OpenXRHMD.cpp b/Engine/Plugins/Runtime/OpenXR/Source/OpenXRHMD/Private/OpenXRHMD.cpp
--- a/Engine/Plugins/Runtime/OpenXR/Source/OpenXRHMD/Private/OpenXRHMD.cpp
+++ b/Engine/Plugins/Runtime/OpenXR/Source/OpenXRHMD/Private/OpenXRHMD.cpp
@@ -825,6 +825,7 @@ bool FOpenXRHMD::EnableStereo(bool stereo)
 
 			return true;
 		}
+		bStereoEnabled = false;
 		return false;
 	}
 	else
