Fix alpha of emulated stereo layers on hmd and spectator screen.

diff --git a/Engine/Plugins/Runtime/OpenXR/Source/OpenXRHMD/Private/OpenXRHMD.cpp b/Engine/Plugins/Runtime/OpenXR/Source/OpenXRHMD/Private/OpenXRHMD.cpp
--- a/Engine/Plugins/Runtime/OpenXR/Source/OpenXRHMD/Private/OpenXRHMD.cpp
+++ b/Engine/Plugins/Runtime/OpenXR/Source/OpenXRHMD/Private/OpenXRHMD.cpp
@@ -2497,7 +2497,8 @@ void FOpenXRHMD::SetupFrameQuadLayers_RenderThread(FRHICommandListImmediate& RHI
 				{
 					FRHITexture2D* SrcTexture = Layer.Desc.Texture->GetTexture2D();
 					const FIntRect DstRect(FIntPoint(0, 0), Layer.RightEye.SwapchainSize.IntPoint());
-					CopyTexture_RenderThread(RHICmdList, SrcTexture, FIntRect(), Layer.RightEye.Swapchain, DstRect, false, bNoAlpha);
+					ETextureAlphaBlendState SrcTextureAlphaBlendState = bNoAlpha ? ETextureAlphaBlendState::NoAlpha : ETextureAlphaBlendState::UnpremultipliedAlpha;
+					CopyTexture_RenderThread(RHICmdList, SrcTexture, FIntRect(), Layer.RightEye.Swapchain, DstRect, false, SrcTextureAlphaBlendState);
 				}
 
 				Quad.eyeVisibility = bIsStereo ? XR_EYE_VISIBILITY_RIGHT : XR_EYE_VISIBILITY_BOTH;
@@ -2515,7 +2516,8 @@ void FOpenXRHMD::SetupFrameQuadLayers_RenderThread(FRHICommandListImmediate& RHI
 				{
 					FRHITexture2D* SrcTexture = Layer.Desc.LeftTexture->GetTexture2D();
 					const FIntRect DstRect(FIntPoint(0, 0), Layer.LeftEye.SwapchainSize.IntPoint());
-					CopyTexture_RenderThread(RHICmdList, SrcTexture, FIntRect(), Layer.LeftEye.Swapchain, DstRect, false, bNoAlpha);
+					ETextureAlphaBlendState SrcTextureAlphaBlendState = bNoAlpha ? ETextureAlphaBlendState::NoAlpha : ETextureAlphaBlendState::UnpremultipliedAlpha;
+					CopyTexture_RenderThread(RHICmdList, SrcTexture, FIntRect(), Layer.LeftEye.Swapchain, DstRect, false, SrcTextureAlphaBlendState);
 				}
 
 				Quad.eyeVisibility = XR_EYE_VISIBILITY_LEFT;
@@ -2558,7 +2560,8 @@ void FOpenXRHMD::DrawEmulatedFaceLockedLayers_RenderThread(FRDGBuilder& GraphBui
 		RHICmdList.BeginRenderPass(RPInfo, TEXT("EmulatedFaceLockedStereoLayerRender"));
 		RHICmdList.SetViewport((float)RenderParams.Viewport.Min.X, (float)RenderParams.Viewport.Min.Y, 0.0f, (float)RenderParams.Viewport.Max.X, (float)RenderParams.Viewport.Max.Y, 1.0f);
 
-		DrawClearQuadAlpha(RHICmdList, 0);
+		// We need to clear to black + 0 alpha in order to composite opaque + transparent layers correctly
+		DrawClearQuad(RHICmdList, FLinearColor::Transparent);
 
 		FDefaultStereoLayers::StereoLayerRender(RHICmdList, EmulatedFaceLockedLayers, RenderParams);
 
@@ -3481,7 +3484,7 @@ private:
 IMPLEMENT_SHADER_TYPE(, FDisplayMappingPS, TEXT("/Engine/Private/CompositeUIPixelShader.usf"), TEXT("DisplayMappingPS"), SF_Pixel);
 
 void FOpenXRHMD::CopyTexture_RenderThread(FRHICommandListImmediate& RHICmdList, FRHITexture2D* SrcTexture, FIntRect SrcRect, FRHITexture2D* DstTexture, FIntRect DstRect, 
-											bool bClearBlack, bool bNoAlpha, ERenderTargetActions RTAction, ERHIAccess FinalDstAccess) const
+											bool bClearBlack, ERenderTargetActions RTAction, ERHIAccess FinalDstAccess, ETextureAlphaBlendState SrcTextureAlphaBlendState) const
 {
 	check(IsInRenderingThread());
 
@@ -3524,7 +3527,24 @@ void FOpenXRHMD::CopyTexture_RenderThread(FRHICommandListImmediate& RHICmdList,
 
 		FGraphicsPipelineStateInitializer GraphicsPSOInit;
 		RHICmdList.ApplyCachedRenderTargets(GraphicsPSOInit);
-		GraphicsPSOInit.BlendState = bNoAlpha ? TStaticBlendState<>::GetRHI() : TStaticBlendState<CW_RGBA, BO_Add, BF_SourceAlpha, BF_InverseSourceAlpha, BO_Add, BF_One, BF_InverseSourceAlpha>::GetRHI();
+
+		// We need to differentiate between types of layers: opaque, unpremultiplied alpha (regular texture copy) and premultiplied alpha (emulation texture)
+		switch (SrcTextureAlphaBlendState)
+		{
+			case ETextureAlphaBlendState::NoAlpha:
+				GraphicsPSOInit.BlendState = TStaticBlendState<>::GetRHI();
+				break;
+			case ETextureAlphaBlendState::PremultipliedAlpha:
+				// Because StereoLayerRender actually enables alpha blending as it composites the layers into the emulation texture
+				// the color values for the emulation swapchain are PREMULTIPLIED ALPHA. That means we don't want to multiply alpha again!
+				// So we can just do SourceColor * 1.0f + DestColor (1 - SourceAlpha)
+				GraphicsPSOInit.BlendState = TStaticBlendState<CW_RGBA, BO_Add, BF_One, BF_InverseSourceAlpha, BO_Add, BF_One, BF_InverseSourceAlpha>::GetRHI();
+				break;
+			case ETextureAlphaBlendState::UnpremultipliedAlpha:
+				GraphicsPSOInit.BlendState = TStaticBlendState<CW_RGBA, BO_Add, BF_SourceAlpha, BF_InverseSourceAlpha, BO_Add, BF_One, BF_InverseSourceAlpha>::GetRHI();
+				break;
+		}
+		
 		GraphicsPSOInit.RasterizerState = TStaticRasterizerState<>::GetRHI();
 		GraphicsPSOInit.DepthStencilState = TStaticDepthStencilState<false, CF_Always>::GetRHI();
 		GraphicsPSOInit.PrimitiveType = PT_TriangleList;
@@ -3634,7 +3654,7 @@ void FOpenXRHMD::CopyTexture_RenderThread(FRHICommandListImmediate& RHICmdList,
 	RHICmdList.Transition(FRHITransitionInfo(DstTexture, ERHIAccess::RTV, FinalDstAccess));
 }
 
-void FOpenXRHMD::CopyTexture_RenderThread(FRHICommandListImmediate& RHICmdList, FRHITexture2D* SrcTexture, FIntRect SrcRect, const FXRSwapChainPtr& DstSwapChain, FIntRect DstRect, bool bClearBlack, bool bNoAlpha) const
+void FOpenXRHMD::CopyTexture_RenderThread(FRHICommandListImmediate& RHICmdList, FRHITexture2D* SrcTexture, FIntRect SrcRect, const FXRSwapChainPtr& DstSwapChain, FIntRect DstRect, bool bClearBlack, ETextureAlphaBlendState SrcTextureAlphaBlendState) const
 {
 	RHICmdList.EnqueueLambda([DstSwapChain](FRHICommandListImmediate& InRHICmdList)
 	{
@@ -3644,7 +3664,7 @@ void FOpenXRHMD::CopyTexture_RenderThread(FRHICommandListImmediate& RHICmdList,
 
 	// Now that we've enqueued the swapchain wait we can add the commands to do the actual texture copy
 	FRHITexture2D* const DstTexture = DstSwapChain->GetTexture2DArray() ? DstSwapChain->GetTexture2DArray() : DstSwapChain->GetTexture2D();
-	CopyTexture_RenderThread(RHICmdList, SrcTexture, SrcRect, DstTexture, DstRect, bClearBlack, bNoAlpha, ERenderTargetActions::Clear_Store, ERHIAccess::SRVMask);
+	CopyTexture_RenderThread(RHICmdList, SrcTexture, SrcRect, DstTexture, DstRect, bClearBlack, ERenderTargetActions::Clear_Store, ERHIAccess::SRVMask, SrcTextureAlphaBlendState);
 
 	// Enqueue a command to release the image after the copy is done
 	RHICmdList.EnqueueLambda([DstSwapChain](FRHICommandListImmediate& InRHICmdList)
@@ -3656,7 +3676,9 @@ void FOpenXRHMD::CopyTexture_RenderThread(FRHICommandListImmediate& RHICmdList,
 void FOpenXRHMD::CopyTexture_RenderThread(FRHICommandListImmediate& RHICmdList, FRHITexture2D* SrcTexture, FIntRect SrcRect, FRHITexture2D* DstTexture, FIntRect DstRect, bool bClearBlack, bool bNoAlpha) const
 {
 	// FIXME: This should probably use the Load_Store action since the spectator controller does multiple overlaying copies.
-	CopyTexture_RenderThread(RHICmdList, SrcTexture, SrcRect, DstTexture, DstRect, bClearBlack, bNoAlpha, ERenderTargetActions::DontLoad_Store, ERHIAccess::Present);
+	// This call only comes from the spectator screen so we expect alpha to be premultiplied.
+	ETextureAlphaBlendState SrcTextureAlphaBlendState = bNoAlpha ? ETextureAlphaBlendState::NoAlpha : ETextureAlphaBlendState::PremultipliedAlpha;
+	CopyTexture_RenderThread(RHICmdList, SrcTexture, SrcRect, DstTexture, DstRect, bClearBlack, ERenderTargetActions::DontLoad_Store, ERHIAccess::Present, SrcTextureAlphaBlendState);
 }
 
 void FOpenXRHMD::RenderTexture_RenderThread(class FRHICommandListImmediate& RHICmdList, class FRHITexture* BackBuffer, class FRHITexture* SrcTexture, FVector2D WindowSize) const
diff --git a/Engine/Plugins/Runtime/OpenXR/Source/OpenXRHMD/Private/OpenXRHMD.h b/Engine/Plugins/Runtime/OpenXR/Source/OpenXRHMD/Private/OpenXRHMD.h
--- a/Engine/Plugins/Runtime/OpenXR/Source/OpenXRHMD/Private/OpenXRHMD.h
+++ b/Engine/Plugins/Runtime/OpenXR/Source/OpenXRHMD/Private/OpenXRHMD.h
@@ -37,6 +37,13 @@ class FOpenXRHMD
 	, public TStereoLayerManager<FOpenXRLayer>
 	, public IOpenXRExtensionPluginDelegates
 {
+private:
+	enum ETextureAlphaBlendState : uint8
+	{
+		NoAlpha,
+		PremultipliedAlpha,
+		UnpremultipliedAlpha,
+	};
 public:
 	class FDeviceSpace
 	{
@@ -253,9 +260,9 @@ protected:
 	bool IsViewManagedByPlugin(int32 ViewIndex) const;
 
 	void CopyTexture_RenderThread(FRHICommandListImmediate& RHICmdList, FRHITexture2D* SrcTexture, FIntRect SrcRect, FRHITexture2D* DstTexture, FIntRect DstRect, 
-								  bool bClearBlack, bool bNoAlpha, ERenderTargetActions RTAction, ERHIAccess FinalDstAccess) const;
+								  bool bClearBlack, ERenderTargetActions RTAction, ERHIAccess FinalDstAccess, ETextureAlphaBlendState SrcTextureAlphaBlendState) const;
 	
-	void CopyTexture_RenderThread(FRHICommandListImmediate& RHICmdList, FRHITexture2D* SrcTexture, FIntRect SrcRect, const FXRSwapChainPtr& DstSwapChain, FIntRect DstRect, bool bClearBlack, bool bNoAlpha) const;
+	void CopyTexture_RenderThread(FRHICommandListImmediate& RHICmdList, FRHITexture2D* SrcTexture, FIntRect SrcRect, const FXRSwapChainPtr& DstSwapChain, FIntRect DstRect, bool bClearBlack, ETextureAlphaBlendState SrcTextureAlphaBlendState) const;
 
 	void AllocateDepthTextureInternal(uint32 Index, uint32 SizeX, uint32 SizeY, uint32 NumSamples, uint32 ArraySize);
 
diff --git a/Engine/Shaders/Private/StereoLayerShader.usf b/Engine/Shaders/Private/StereoLayerShader.usf
--- a/Engine/Shaders/Private/StereoLayerShader.usf
+++ b/Engine/Shaders/Private/StereoLayerShader.usf
@@ -31,13 +31,21 @@ void MainVS(
 
 Texture2D InTexture;
 SamplerState InTextureSampler;
+float InIsOpaque;
+
+float CalculateAlpha(float4 OutColor)
+{
+	// Don't export alpha for opaque layers
+	return (OutColor.a * (1.0 - InIsOpaque) + InIsOpaque);
+}
 
 void MainPS_Texture2D(
 	    float2 UV       : TEXCOORD0,
 	out float4 OutColor : SV_Target0
 	)
 {
-    OutColor = Texture2DSample(InTexture, InTextureSampler, UV);
+	OutColor = Texture2DSample(InTexture, InTextureSampler, UV);
+	OutColor.a = CalculateAlpha(OutColor);
 }
 
 TextureExternal InExternalTexture;
@@ -47,5 +55,6 @@ void MainPS_TextureExternal(
 	out float4 OutColor : SV_Target0
 	)
 {
-    OutColor = TextureExternalSample(InExternalTexture, InTextureSampler, UV);
+	OutColor = TextureExternalSample(InExternalTexture, InTextureSampler, UV);
+	OutColor.a = CalculateAlpha(OutColor);
 }
diff --git a/Engine/Source/Runtime/Engine/Public/StereoLayerRendering.h b/Engine/Source/Runtime/Engine/Public/StereoLayerRendering.h
--- a/Engine/Source/Runtime/Engine/Public/StereoLayerRendering.h
+++ b/Engine/Source/Runtime/Engine/Public/StereoLayerRendering.h
@@ -76,11 +76,17 @@ public:
 
 	static bool ShouldCompilePermutation(const FGlobalShaderPermutationParameters& Parameters) { return true; }
 
-	void SetParameters(FRHICommandList& RHICmdList, FRHISamplerState* SamplerStateRHI, FRHITexture* TextureRHI)
+	void SetParameters(FRHICommandList& RHICmdList, FRHISamplerState* SamplerStateRHI, FRHITexture* TextureRHI, bool bIsOpaque)
 	{
 		FRHIPixelShader* PS = RHICmdList.GetBoundPixelShader();
 
 		SetTextureParameter(RHICmdList, PS, InTexture, InTextureSampler, SamplerStateRHI, TextureRHI);
+
+		if (InIsOpaque.IsBound())
+		{
+			const float OpaqueVal = bIsOpaque ? 1.0 : 0.0;
+			SetShaderValue(RHICmdList, PS, InIsOpaque, OpaqueVal);
+		}
 	}
 
 protected:
@@ -89,11 +95,13 @@ protected:
 	{
 		InTexture.Bind(Initializer.ParameterMap, TextureParamName, SPF_Mandatory);
 		InTextureSampler.Bind(Initializer.ParameterMap, TEXT("InTextureSampler"));
+		InIsOpaque.Bind(Initializer.ParameterMap, TEXT("InIsOpaque"));
 	}
 	FStereoLayerPS_Base() {}
 
 	LAYOUT_FIELD(FShaderResourceParameter, InTexture);
 	LAYOUT_FIELD(FShaderResourceParameter, InTextureSampler);
+	LAYOUT_FIELD(FShaderParameter, InIsOpaque);
 };
 
 /**
diff --git a/Engine/Source/Runtime/HeadMountedDisplay/Private/DefaultStereoLayers.cpp b/Engine/Source/Runtime/HeadMountedDisplay/Private/DefaultStereoLayers.cpp
--- a/Engine/Source/Runtime/HeadMountedDisplay/Private/DefaultStereoLayers.cpp
+++ b/Engine/Source/Runtime/HeadMountedDisplay/Private/DefaultStereoLayers.cpp
@@ -151,7 +151,8 @@ void FDefaultStereoLayers::StereoLayerRender(FRHICommandListImmediate& RHICmdLis
 		PixelShader->SetParameters(
 			RHICmdList,
 			TStaticSamplerState<SF_Trilinear>::GetRHI(),
-			Layer.Texture);
+			Layer.Texture,
+			bIsOpaque);
 
 		const FIntPoint TargetSize = RenderParams.Viewport.Size();
 		// Draw primitive
