OpenXR: Fix off-by-one error resulting in trailing NUL character embedded in FNames returned by `FOpenXRHMDModule::ResolvePathToName`.

diff --git a/Engine/Plugins/Runtime/OpenXR/Source/OpenXRHMD/Private/OpenXRHMDModule.cpp b/Engine/Plugins/Runtime/OpenXR/Source/OpenXRHMD/Private/OpenXRHMDModule.cpp
--- a/Engine/Plugins/Runtime/OpenXR/Source/OpenXRHMD/Private/OpenXRHMDModule.cpp
+++ b/Engine/Plugins/Runtime/OpenXR/Source/OpenXRHMD/Private/OpenXRHMDModule.cpp
@@ -790,7 +790,7 @@ FName FOpenXRHMDModule::ResolvePathToName(XrPath Path)
 	if (Result == XR_SUCCESS)
 	{
 		// Resolve this XrPath to an FName and store it in the name map
-		FName Name(PathCount, PathChars);
+		FName Name(PathCount - 1, PathChars);
 
 		FWriteScopeLock Lock(NameMutex);
 		PathToName.Add(Path, Name);
