Fix a missing shader when copying swapchain image for the spectator.

- Do not create texture array for just one element to not even go that path if MMV isn't enabled.

diff --git a/Engine/Plugins/Runtime/OpenXR/Source/OpenXRHMD/Private/OpenXRHMD_Swapchain.cpp b/Engine/Plugins/Runtime/OpenXR/Source/OpenXRHMD/Private/OpenXRHMD_Swapchain.cpp
--- a/Engine/Plugins/Runtime/OpenXR/Source/OpenXRHMD/Private/OpenXRHMD_Swapchain.cpp
+++ b/Engine/Plugins/Runtime/OpenXR/Source/OpenXRHMD/Private/OpenXRHMD_Swapchain.cpp
@@ -283,7 +283,10 @@ FXRSwapChainPtr CreateSwapchain_D3D11(XrSession InSession, uint8 Format, uint8&
 	TArray<XrSwapchainImageD3D11KHR> Images = EnumerateImages<XrSwapchainImageD3D11KHR>(Swapchain, XR_TYPE_SWAPCHAIN_IMAGE_D3D11_KHR);
 	for (const auto& Image : Images)
 	{
-		TextureChain.Add(static_cast<FTextureRHIRef>(D3D11RHI->RHICreateTexture2DArrayFromResource(GPixelFormats[Format].UnrealFormat, CreateFlags, ClearValueBinding, Image.texture)));
+		TextureChain.Add(static_cast<FTextureRHIRef>((ArraySize > 1) ?
+			D3D11RHI->RHICreateTexture2DArrayFromResource(GPixelFormats[Format].UnrealFormat, CreateFlags, ClearValueBinding, Image.texture) :
+			D3D11RHI->RHICreateTexture2DFromResource(GPixelFormats[Format].UnrealFormat, CreateFlags, ClearValueBinding, Image.texture)
+			));
 	}
 	FTextureRHIRef ChainTarget = static_cast<FTextureRHIRef>(D3D11RHI->RHICreateAliasedTexture((FTextureRHIRef&)TextureChain[0]));
 
@@ -346,7 +349,9 @@ FXRSwapChainPtr CreateSwapchain_OpenGL(XrSession InSession, uint8 Format, uint8&
 	TArray<XrSwapchainImageOpenGLKHR> Images = EnumerateImages<XrSwapchainImageOpenGLKHR>(Swapchain, XR_TYPE_SWAPCHAIN_IMAGE_OPENGL_KHR);
 	for (const auto& Image : Images)
 	{
-		FTexture2DRHIRef NewTexture = DynamicRHI->RHICreateTexture2DArrayFromResource(GPixelFormats[Format].UnrealFormat, SizeX, SizeY, ArraySize, NumMips, NumSamples, 1, ClearValueBinding, Image.image, CreateFlags);
+		FTexture2DRHIRef NewTexture = (ArraySize > 1) ?
+			DynamicRHI->RHICreateTexture2DArrayFromResource(GPixelFormats[Format].UnrealFormat, SizeX, SizeY, ArraySize, NumMips, NumSamples, 1, ClearValueBinding, Image.image, CreateFlags) :
+			DynamicRHI->RHICreateTexture2DFromResource(GPixelFormats[Format].UnrealFormat, SizeX, SizeY, NumMips, NumSamples, 1, ClearValueBinding, Image.image, CreateFlags);
 		TextureChain.Add(NewTexture.GetReference());
 	}
 	FTextureRHIRef ChainTarget = static_cast<FTextureRHIRef>(DynamicRHI->RHICreateAliasedTexture((FTextureRHIRef&)TextureChain[0]));
@@ -376,7 +381,9 @@ FXRSwapChainPtr CreateSwapchain_OpenGLES(XrSession InSession, uint8 Format, uint
 	TArray<XrSwapchainImageOpenGLESKHR> Images = EnumerateImages<XrSwapchainImageOpenGLESKHR>(Swapchain, XR_TYPE_SWAPCHAIN_IMAGE_OPENGL_ES_KHR);
 	for (const auto& Image : Images)
 	{
-		FTexture2DRHIRef NewTexture = DynamicRHI->RHICreateTexture2DArrayFromResource(GPixelFormats[Format].UnrealFormat, SizeX, SizeY, ArraySize, NumMips, NumSamples, 1, ClearValueBinding, Image.image, CreateFlags);
+		FTexture2DRHIRef NewTexture = (ArraySize > 1) ?
+			DynamicRHI->RHICreateTexture2DArrayFromResource(GPixelFormats[Format].UnrealFormat, SizeX, SizeY, ArraySize, NumMips, NumSamples, 1, ClearValueBinding, Image.image, CreateFlags) :
+			DynamicRHI->RHICreateTexture2DFromResource(GPixelFormats[Format].UnrealFormat, SizeX, SizeY, NumMips, NumSamples, 1, ClearValueBinding, Image.image, CreateFlags);
 		TextureChain.Add(NewTexture.GetReference());
 	}
 	FTextureRHIRef ChainTarget = static_cast<FTextureRHIRef>(DynamicRHI->RHICreateAliasedTexture((FTextureRHIRef&)TextureChain[0]));
@@ -427,7 +434,10 @@ FXRSwapChainPtr CreateSwapchain_Vulkan(XrSession InSession, uint8 Format, uint8&
 	TArray<XrSwapchainImageVulkanKHR> Images = EnumerateImages<XrSwapchainImageVulkanKHR>(Swapchain, XR_TYPE_SWAPCHAIN_IMAGE_VULKAN_KHR);
 	for (const auto& Image : Images)
 	{
-		TextureChain.Add(static_cast<FTextureRHIRef>(VulkanRHI->RHICreateTexture2DArrayFromResource(GPixelFormats[Format].UnrealFormat, SizeX, SizeY, ArraySize, NumMips, NumSamples, Image.image, CreateFlags, ClearValueBinding)));
+		TextureChain.Add(static_cast<FTextureRHIRef>((ArraySize > 1) ?
+			VulkanRHI->RHICreateTexture2DArrayFromResource(GPixelFormats[Format].UnrealFormat, SizeX, SizeY, ArraySize, NumMips, NumSamples, Image.image, CreateFlags, ClearValueBinding) :
+			VulkanRHI->RHICreateTexture2DFromResource(GPixelFormats[Format].UnrealFormat, SizeX, SizeY, NumMips, NumSamples, Image.image, CreateFlags, ClearValueBinding)
+			));
 	}
 	FTextureRHIRef ChainTarget = static_cast<FTextureRHIRef>(VulkanRHI->RHICreateAliasedTexture((FTextureRHIRef&)TextureChain[0]));
 
