OpenXR: refine CopyTexture operations for layer + spectator screen usage

Native + emulated layers interact with CopyTexture in different ways, and we need to make sure we clear color + alpha correctly, write out the correct color channels, and manage blending when needed.

diff --git a/Engine/Plugins/Runtime/OpenXR/Source/OpenXRHMD/Private/OpenXRHMD.cpp b/Engine/Plugins/Runtime/OpenXR/Source/OpenXRHMD/Private/OpenXRHMD.cpp
--- a/Engine/Plugins/Runtime/OpenXR/Source/OpenXRHMD/Private/OpenXRHMD.cpp
+++ b/Engine/Plugins/Runtime/OpenXR/Source/OpenXRHMD/Private/OpenXRHMD.cpp
@@ -130,6 +130,13 @@ namespace {
 // OpenXRHMD IHeadMountedDisplay Implementation
 //---------------------------------------------------
 
+enum FOpenXRHMD::ETextureCopyBlendModifier : uint8
+{
+	Opaque,
+	TransparentAlphaPassthrough,
+	PremultipliedAlphaBlend,
+};
+
 bool FOpenXRHMD::FVulkanExtensions::GetVulkanInstanceExtensionsRequired(TArray<const ANSICHAR*>& Out)
 {
 #ifdef XR_USE_GRAPHICS_API_VULKAN
@@ -2489,6 +2496,7 @@ void FOpenXRHMD::SetupFrameQuadLayers_RenderThread(FRHICommandListImmediate& RHI
 
 			// The layer pose doesn't take the transform scale into consideration, so we need to manually apply it the quad size.
 			const FVector2D LayerComponentScaler(Layer.Desc.Transform.GetScale3D().Y, Layer.Desc.Transform.GetScale3D().Z);
+			const ETextureCopyBlendModifier SrcTextureCopyModifier = bNoAlpha ? ETextureCopyBlendModifier::Opaque : ETextureCopyBlendModifier::TransparentAlphaPassthrough;
 
 			// We need to copy each layer into an OpenXR swapchain so they can be displayed by the compositor
 			if (Layer.RightEye.Swapchain.IsValid() && Layer.Desc.Texture.IsValid())
@@ -2497,8 +2505,7 @@ void FOpenXRHMD::SetupFrameQuadLayers_RenderThread(FRHICommandListImmediate& RHI
 				{
 					FRHITexture2D* SrcTexture = Layer.Desc.Texture->GetTexture2D();
 					const FIntRect DstRect(FIntPoint(0, 0), Layer.RightEye.SwapchainSize.IntPoint());
-					ETextureAlphaBlendState SrcTextureAlphaBlendState = bNoAlpha ? ETextureAlphaBlendState::NoAlpha : ETextureAlphaBlendState::UnpremultipliedAlpha;
-					CopyTexture_RenderThread(RHICmdList, SrcTexture, FIntRect(), Layer.RightEye.Swapchain, DstRect, false, SrcTextureAlphaBlendState);
+					CopyTexture_RenderThread(RHICmdList, SrcTexture, FIntRect(), Layer.RightEye.Swapchain, DstRect, false, SrcTextureCopyModifier);
 				}
 
 				Quad.eyeVisibility = bIsStereo ? XR_EYE_VISIBILITY_RIGHT : XR_EYE_VISIBILITY_BOTH;
@@ -2516,8 +2523,7 @@ void FOpenXRHMD::SetupFrameQuadLayers_RenderThread(FRHICommandListImmediate& RHI
 				{
 					FRHITexture2D* SrcTexture = Layer.Desc.LeftTexture->GetTexture2D();
 					const FIntRect DstRect(FIntPoint(0, 0), Layer.LeftEye.SwapchainSize.IntPoint());
-					ETextureAlphaBlendState SrcTextureAlphaBlendState = bNoAlpha ? ETextureAlphaBlendState::NoAlpha : ETextureAlphaBlendState::UnpremultipliedAlpha;
-					CopyTexture_RenderThread(RHICmdList, SrcTexture, FIntRect(), Layer.LeftEye.Swapchain, DstRect, false, SrcTextureAlphaBlendState);
+					CopyTexture_RenderThread(RHICmdList, SrcTexture, FIntRect(), Layer.LeftEye.Swapchain, DstRect, false, SrcTextureCopyModifier);
 				}
 
 				Quad.eyeVisibility = XR_EYE_VISIBILITY_LEFT;
@@ -3483,8 +3489,20 @@ private:
 
 IMPLEMENT_SHADER_TYPE(, FDisplayMappingPS, TEXT("/Engine/Private/CompositeUIPixelShader.usf"), TEXT("DisplayMappingPS"), SF_Pixel);
 
+// We use CopyTexture for a number of subtley different use cases.
+// * SpectatorScreen, background layer - optional clear to black, opaque
+// * SpectatorScreen, emulated facelocked layer - no clears, blend premultiplied alpha onto background layer
+// * Native layer texture init, ignore alpha component - no color clear, opaque
+// * Native layer texture init, export alpha component - no color clear, export source texture alpha
+//
+// With this information, we can extract some common usages
+// * opaque copies should clear alpha to 1.0f, and write RGB
+//	* background layer can optionally clear to black
+// * transparent export just writes thru RGBA, no clears needed
+// * premultiplied alpha blend needs custom blend, no clears needed
+
 void FOpenXRHMD::CopyTexture_RenderThread(FRHICommandListImmediate& RHICmdList, FRHITexture2D* SrcTexture, FIntRect SrcRect, FRHITexture2D* DstTexture, FIntRect DstRect, 
-											bool bClearBlack, ERenderTargetActions RTAction, ERHIAccess FinalDstAccess, ETextureAlphaBlendState SrcTextureAlphaBlendState) const
+											bool bClearBlack, ERenderTargetActions RTAction, ERHIAccess FinalDstAccess, ETextureCopyBlendModifier SrcTextureCopyModifier) const
 {
 	check(IsInRenderingThread());
 
@@ -3516,11 +3534,20 @@ void FOpenXRHMD::CopyTexture_RenderThread(FRHICommandListImmediate& RHICmdList,
 	FRHIRenderPassInfo RenderPassInfo(ColorRT, RTAction);
 	RHICmdList.BeginRenderPass(RenderPassInfo, TEXT("OpenXRHMD_CopyTexture"));
 	{
-		if (bClearBlack)
+		if (bClearBlack || SrcTextureCopyModifier == ETextureCopyBlendModifier::Opaque)
 		{
 			const FIntRect ClearRect(0, 0, DstTexture->GetSizeX(), DstTexture->GetSizeY());
 			RHICmdList.SetViewport(ClearRect.Min.X, ClearRect.Min.Y, 0, ClearRect.Max.X, ClearRect.Max.Y, 1.0f);
-			DrawClearQuad(RHICmdList, FLinearColor::Black);
+
+			if (bClearBlack)
+			{
+				DrawClearQuad(RHICmdList, FLinearColor::Black);
+			}
+			else
+			{
+				// For opaque texture copies, we want to make sure alpha is initialized to 1.0f
+				DrawClearQuadAlpha(RHICmdList, 1.0f);
+			}
 		}
 
 		RHICmdList.SetViewport(DstRect.Min.X, DstRect.Min.Y, 0, DstRect.Max.X, DstRect.Max.Y, 1.0f);
@@ -3529,19 +3556,23 @@ void FOpenXRHMD::CopyTexture_RenderThread(FRHICommandListImmediate& RHICmdList,
 		RHICmdList.ApplyCachedRenderTargets(GraphicsPSOInit);
 
 		// We need to differentiate between types of layers: opaque, unpremultiplied alpha (regular texture copy) and premultiplied alpha (emulation texture)
-		switch (SrcTextureAlphaBlendState)
+		switch (SrcTextureCopyModifier)
 		{
-			case ETextureAlphaBlendState::NoAlpha:
-				GraphicsPSOInit.BlendState = TStaticBlendState<>::GetRHI();
+			case ETextureCopyBlendModifier::Opaque:
+				GraphicsPSOInit.BlendState = TStaticBlendState<CW_RGB>::GetRHI();
+				break;
+			case ETextureCopyBlendModifier::TransparentAlphaPassthrough:
+				GraphicsPSOInit.BlendState = TStaticBlendState<CW_RGBA>::GetRHI();
 				break;
-			case ETextureAlphaBlendState::PremultipliedAlpha:
+			case ETextureCopyBlendModifier::PremultipliedAlphaBlend:
 				// Because StereoLayerRender actually enables alpha blending as it composites the layers into the emulation texture
 				// the color values for the emulation swapchain are PREMULTIPLIED ALPHA. That means we don't want to multiply alpha again!
 				// So we can just do SourceColor * 1.0f + DestColor (1 - SourceAlpha)
 				GraphicsPSOInit.BlendState = TStaticBlendState<CW_RGBA, BO_Add, BF_One, BF_InverseSourceAlpha, BO_Add, BF_One, BF_InverseSourceAlpha>::GetRHI();
 				break;
-			case ETextureAlphaBlendState::UnpremultipliedAlpha:
-				GraphicsPSOInit.BlendState = TStaticBlendState<CW_RGBA, BO_Add, BF_SourceAlpha, BF_InverseSourceAlpha, BO_Add, BF_One, BF_InverseSourceAlpha>::GetRHI();
+			default:
+				check(!"Unsupported copy modifier");
+				GraphicsPSOInit.BlendState = TStaticBlendState<>::GetRHI();
 				break;
 		}
 		
@@ -3654,7 +3685,7 @@ void FOpenXRHMD::CopyTexture_RenderThread(FRHICommandListImmediate& RHICmdList,
 	RHICmdList.Transition(FRHITransitionInfo(DstTexture, ERHIAccess::RTV, FinalDstAccess));
 }
 
-void FOpenXRHMD::CopyTexture_RenderThread(FRHICommandListImmediate& RHICmdList, FRHITexture2D* SrcTexture, FIntRect SrcRect, const FXRSwapChainPtr& DstSwapChain, FIntRect DstRect, bool bClearBlack, ETextureAlphaBlendState SrcTextureAlphaBlendState) const
+void FOpenXRHMD::CopyTexture_RenderThread(FRHICommandListImmediate& RHICmdList, FRHITexture2D* SrcTexture, FIntRect SrcRect, const FXRSwapChainPtr& DstSwapChain, FIntRect DstRect, bool bClearBlack, ETextureCopyBlendModifier SrcTextureCopyModifier) const
 {
 	RHICmdList.EnqueueLambda([DstSwapChain](FRHICommandListImmediate& InRHICmdList)
 	{
@@ -3664,7 +3695,7 @@ void FOpenXRHMD::CopyTexture_RenderThread(FRHICommandListImmediate& RHICmdList,
 
 	// Now that we've enqueued the swapchain wait we can add the commands to do the actual texture copy
 	FRHITexture2D* const DstTexture = DstSwapChain->GetTexture2DArray() ? DstSwapChain->GetTexture2DArray() : DstSwapChain->GetTexture2D();
-	CopyTexture_RenderThread(RHICmdList, SrcTexture, SrcRect, DstTexture, DstRect, bClearBlack, ERenderTargetActions::Clear_Store, ERHIAccess::SRVMask, SrcTextureAlphaBlendState);
+	CopyTexture_RenderThread(RHICmdList, SrcTexture, SrcRect, DstTexture, DstRect, bClearBlack, ERenderTargetActions::Clear_Store, ERHIAccess::SRVMask, SrcTextureCopyModifier);
 
 	// Enqueue a command to release the image after the copy is done
 	RHICmdList.EnqueueLambda([DstSwapChain](FRHICommandListImmediate& InRHICmdList)
@@ -3675,10 +3706,9 @@ void FOpenXRHMD::CopyTexture_RenderThread(FRHICommandListImmediate& RHICmdList,
 
 void FOpenXRHMD::CopyTexture_RenderThread(FRHICommandListImmediate& RHICmdList, FRHITexture2D* SrcTexture, FIntRect SrcRect, FRHITexture2D* DstTexture, FIntRect DstRect, bool bClearBlack, bool bNoAlpha) const
 {
-	// FIXME: This should probably use the Load_Store action since the spectator controller does multiple overlaying copies.
 	// This call only comes from the spectator screen so we expect alpha to be premultiplied.
-	ETextureAlphaBlendState SrcTextureAlphaBlendState = bNoAlpha ? ETextureAlphaBlendState::NoAlpha : ETextureAlphaBlendState::PremultipliedAlpha;
-	CopyTexture_RenderThread(RHICmdList, SrcTexture, SrcRect, DstTexture, DstRect, bClearBlack, ERenderTargetActions::DontLoad_Store, ERHIAccess::Present, SrcTextureAlphaBlendState);
+	const ETextureCopyBlendModifier SrcTextureCopyModifier = bNoAlpha ? ETextureCopyBlendModifier::Opaque : ETextureCopyBlendModifier::PremultipliedAlphaBlend;
+	CopyTexture_RenderThread(RHICmdList, SrcTexture, SrcRect, DstTexture, DstRect, bClearBlack, ERenderTargetActions::Load_Store, ERHIAccess::Present, SrcTextureCopyModifier);
 }
 
 void FOpenXRHMD::RenderTexture_RenderThread(class FRHICommandListImmediate& RHICmdList, class FRHITexture* BackBuffer, class FRHITexture* SrcTexture, FVector2D WindowSize) const
diff --git a/Engine/Plugins/Runtime/OpenXR/Source/OpenXRHMD/Private/OpenXRHMD.h b/Engine/Plugins/Runtime/OpenXR/Source/OpenXRHMD/Private/OpenXRHMD.h
--- a/Engine/Plugins/Runtime/OpenXR/Source/OpenXRHMD/Private/OpenXRHMD.h
+++ b/Engine/Plugins/Runtime/OpenXR/Source/OpenXRHMD/Private/OpenXRHMD.h
@@ -38,12 +38,7 @@ class FOpenXRHMD
 	, public IOpenXRExtensionPluginDelegates
 {
 private:
-	enum ETextureAlphaBlendState : uint8
-	{
-		NoAlpha,
-		PremultipliedAlpha,
-		UnpremultipliedAlpha,
-	};
+
 public:
 	class FDeviceSpace
 	{
@@ -239,6 +234,8 @@ public:
 
 protected:
 
+	enum ETextureCopyBlendModifier : uint8;
+
 	bool StartSession();
 	bool StopSession();
 	bool OnStereoStartup();
@@ -260,9 +257,9 @@ protected:
 	bool IsViewManagedByPlugin(int32 ViewIndex) const;
 
 	void CopyTexture_RenderThread(FRHICommandListImmediate& RHICmdList, FRHITexture2D* SrcTexture, FIntRect SrcRect, FRHITexture2D* DstTexture, FIntRect DstRect, 
-								  bool bClearBlack, ERenderTargetActions RTAction, ERHIAccess FinalDstAccess, ETextureAlphaBlendState SrcTextureAlphaBlendState) const;
+								  bool bClearBlack, ERenderTargetActions RTAction, ERHIAccess FinalDstAccess, ETextureCopyBlendModifier SrcTextureCopyModifier) const;
 	
-	void CopyTexture_RenderThread(FRHICommandListImmediate& RHICmdList, FRHITexture2D* SrcTexture, FIntRect SrcRect, const FXRSwapChainPtr& DstSwapChain, FIntRect DstRect, bool bClearBlack, ETextureAlphaBlendState SrcTextureAlphaBlendState) const;
+	void CopyTexture_RenderThread(FRHICommandListImmediate& RHICmdList, FRHITexture2D* SrcTexture, FIntRect SrcRect, const FXRSwapChainPtr& DstSwapChain, FIntRect DstRect, bool bClearBlack, ETextureCopyBlendModifier SrcTextureCopyModifier) const;
 
 	void AllocateDepthTextureInternal(uint32 Index, uint32 SizeX, uint32 SizeY, uint32 NumSamples, uint32 ArraySize);
 
