Fix empty project name resulting in the project version being used as the application name

We don't need to read the project name from the project settings anymore since the initialization order changed in UE5

diff --git a/Engine/Plugins/Runtime/OpenXR/Source/OpenXRHMD/Private/OpenXRHMDModule.cpp b/Engine/Plugins/Runtime/OpenXR/Source/OpenXRHMD/Private/OpenXRHMDModule.cpp
--- a/Engine/Plugins/Runtime/OpenXR/Source/OpenXRHMD/Private/OpenXRHMDModule.cpp
+++ b/Engine/Plugins/Runtime/OpenXR/Source/OpenXRHMD/Private/OpenXRHMDModule.cpp
@@ -672,16 +672,10 @@ bool FOpenXRHMDModule::InitInstance()
 	auto* CVarDisableEngineAndAppRegistration = IConsoleManager::Get().FindTConsoleVariableDataInt(TEXT("r.DisableEngineAndAppRegistration"));
 	bool bDisableEngineRegistration = (CVarDisableEngineAndAppRegistration && CVarDisableEngineAndAppRegistration->GetValueOnAnyThread() != 0);
 
-	FText ProjectName = FText();
-	GConfig->GetText(TEXT("/Script/EngineSettings.GeneralProjectSettings"), TEXT("ProjectName"), ProjectName, GGameIni);
-
-	FText ProjectVersion = FText();
-	GConfig->GetText(TEXT("/Script/EngineSettings.GeneralProjectSettings"), TEXT("ProjectVersion"), ProjectVersion, GGameIni);
-
 	// EngineName will be of the form "UnrealEngine4.21", with the minor version ("21" in this example)
 	// updated with every quarterly release
 	FString EngineName = bDisableEngineRegistration ? FString("") : FApp::GetEpicProductIdentifier() + FEngineVersion::Current().ToString(EVersionComponent::Minor);
-	FString AppName = bDisableEngineRegistration ? TEXT("") : ProjectName.ToString() + ProjectVersion.ToString();
+	FString AppName = bDisableEngineRegistration ? FString("") : FApp::GetProjectName();
 
 	XrInstanceCreateInfo Info;
 	Info.type = XR_TYPE_INSTANCE_CREATE_INFO;
