UE-171662 Unable to use OpenXR with non-Oculus Android platforms
-Allow android apps with openxr to startup openxr if they are not configured for Oculus Mobile Devices.
-vr.DisableOpenXROnAndroidWithoutOculus added in case the previous behavior is useful for an cross platform project that supports arcore and OpenXR.

diff --git a/Engine/Plugins/Runtime/OpenXR/Source/OpenXRHMD/Private/OpenXRHMDModule.cpp b/Engine/Plugins/Runtime/OpenXR/Source/OpenXRHMD/Private/OpenXRHMDModule.cpp
--- a/Engine/Plugins/Runtime/OpenXR/Source/OpenXRHMD/Private/OpenXRHMDModule.cpp
+++ b/Engine/Plugins/Runtime/OpenXR/Source/OpenXRHMD/Private/OpenXRHMDModule.cpp
@@ -25,6 +25,12 @@ static TAutoConsoleVariable<int32> CVarEnableOpenXRValidationLayer(
 	TEXT("Changes will only take effect in new game/editor instances - can't be changed at runtime.\n"),
 	ECVF_Default);		// @todo: Should we specify ECVF_Cheat here so this doesn't show up in release builds?
 
+static TAutoConsoleVariable<bool> CVarDisableOpenXROnAndroidWithoutOculus(
+	TEXT("vr.DisableOpenXROnAndroidWithoutOculus"),
+	false,
+	TEXT("If true OpenXR will not initialize on Android unless the project is packaged for OpenXR.  (Potentially useful for cross platform projects invovling openxr and arcore, though this is achievable in other ways.)"),
+	ECVF_RenderThreadSafe);
+
 //---------------------------------------------------
 // OpenXRHMD Plugin Implementation
 //---------------------------------------------------
@@ -506,15 +512,21 @@ bool FOpenXRHMDModule::InitInstance()
 	}
 
 #if PLATFORM_ANDROID
-	// TODO: Allow OpenXR on non-Oculus Android platforms
 	if (AndroidThunkCpp_IsOculusMobileApplication())
 	{
-		UE_LOG(LogHMD, Log, TEXT("OpenXRHMDModule: App is packaged for Oculus Mobile"));
+		UE_LOG(LogHMD, Log, TEXT("OpenXRHMDModule: App is packaged for Oculus Mobile OpenXR"));
 	}
 	else
 	{
-		UE_LOG(LogHMD, Log, TEXT("OpenXRHMDModule: App is not packaged for Oculus Mobile"));
-		return false;
+		if (CVarDisableOpenXROnAndroidWithoutOculus.GetValueOnAnyThread())
+		{
+			UE_LOG(LogHMD, Log, TEXT("OpenXRHMDModule: vr.DisableOpenXROnAndroidWithoutOculus is true and this project is not packaged for Oculus Mobile Devices.  Disabling OpenXR."));
+			return false;
+		}
+		else
+		{
+			UE_LOG(LogHMD, Log, TEXT("OpenXRHMDModule: App is packaged for Android OpenXR"));
+		}
 	}
 #endif
 
@@ -529,6 +541,7 @@ bool FOpenXRHMDModule::InitInstance()
 		if (Plugin->GetCustomLoader(&GetProcAddr))
 		{
 			// We pick the first loader we can find
+			UE_LOG(LogHMD, Log, TEXT("OpenXRHMDModule::InitInstance found and will use CustomLoader from plugin %s"), *Plugin->GetDisplayName());
 			break;
 		}
 
@@ -538,6 +551,7 @@ bool FOpenXRHMDModule::InitInstance()
 
 	if (!GetProcAddr)
 	{
+		UE_LOG(LogHMD, Log, TEXT("OpenXRHMDModule::InitInstance using DefaultLoader."));
 		GetProcAddr = GetDefaultLoader();
 	}
 
