Fix assert when recentering without the headset being worn

diff --git a/Engine/Plugins/Runtime/OpenXR/Source/OpenXRHMD/Private/OpenXRHMD.cpp b/Engine/Plugins/Runtime/OpenXR/Source/OpenXRHMD/Private/OpenXRHMD.cpp
--- a/Engine/Plugins/Runtime/OpenXR/Source/OpenXRHMD/Private/OpenXRHMD.cpp
+++ b/Engine/Plugins/Runtime/OpenXR/Source/OpenXRHMD/Private/OpenXRHMD.cpp
@@ -701,9 +701,12 @@ void FOpenXRHMD::ResetPosition()
 
 void FOpenXRHMD::Recenter(EOrientPositionSelector::Type Selector, float Yaw)
 {
-	const FPipelinedFrameState& PipelineState = GetPipelinedFrameStateForThread();
-	const XrTime TargetTime = PipelineState.FrameState.predictedDisplayTime;
-	check(PipelineState.bXrFrameStateUpdated);
+	const XrTime TargetTime = GetDisplayTime();
+	if (TargetTime == 0)
+	{
+		UE_LOG(LogHMD, Warning, TEXT("Could not retrieve a valid head pose for recentering."));
+		return;
+	}
 
 	XrSpace DeviceSpace = XR_NULL_HANDLE;
 	{
