Removing redundant OpenXR include paths. Cleaning up some other include paths to make them easier to track.

diff --git a/Engine/Plugins/Runtime/HPMotionController/Source/HPMotionController/HPMotionController.Build.cs b/Engine/Plugins/Runtime/HPMotionController/Source/HPMotionController/HPMotionController.Build.cs
--- a/Engine/Plugins/Runtime/HPMotionController/Source/HPMotionController/HPMotionController.Build.cs
+++ b/Engine/Plugins/Runtime/HPMotionController/Source/HPMotionController/HPMotionController.Build.cs
@@ -1,8 +1,5 @@
 // Copyright Epic Games, Inc. All Rights Reserved.
 
-using UnrealBuildTool;
-using System.IO;
-
 namespace UnrealBuildTool.Rules
 {
     public class HPMotionController : ModuleRules
@@ -10,13 +7,6 @@ namespace UnrealBuildTool.Rules
         public HPMotionController(ReadOnlyTargetRules Target) 
 				: base(Target)
         {
-			var EngineDir = Path.GetFullPath(Target.RelativeEnginePath);
-            PrivateIncludePaths.AddRange(
-                new string[] {
-                    EngineDir + "/Source/ThirdParty/OpenXR/include"
-				}
-                );
-
             PrivateDependencyModuleNames.AddRange(
                 new string[]
                 {
diff --git a/Engine/Plugins/Runtime/MobileFSR/Source/MobileFSR/MobileFSR.Build.cs b/Engine/Plugins/Runtime/MobileFSR/Source/MobileFSR/MobileFSR.Build.cs
--- a/Engine/Plugins/Runtime/MobileFSR/Source/MobileFSR/MobileFSR.Build.cs
+++ b/Engine/Plugins/Runtime/MobileFSR/Source/MobileFSR/MobileFSR.Build.cs
@@ -1,5 +1,6 @@
 // Copyright Epic Games, Inc. All Rights Reserved.
 
+using System.IO;
 using UnrealBuildTool;
 
 public class MobileFSR : ModuleRules
@@ -10,7 +11,7 @@ public class MobileFSR : ModuleRules
 
 		PrivateIncludePaths.AddRange(
 			new string[] {
-				EngineDirectory + "/Source/Runtime/Renderer/Private",
+				Path.Combine(GetModuleDirectory("Renderer"), "Private"),
 				PluginDirectory + "/Shaders/Private",
 			}
 			);
@@ -29,11 +30,5 @@ public class MobileFSR : ModuleRules
 
 
 		PrecompileForTargets = PrecompileTargetsType.Any;
-
-		if (Target.Platform == UnrealTargetPlatform.Android)
-		{
-			//string PluginPath = Utils.MakePathRelativeTo(ModuleDirectory, Target.RelativeEnginePath);
-			//AdditionalPropertiesForReceipt.Add("AndroidPlugin", Path.Combine(PluginPath, "MobileFSR_APL.xml"));
-		}
 	}
 }
diff --git a/Engine/Plugins/Runtime/Oculus/OculusVR/Source/OculusOpenXRHMD/OculusOpenXRHMD.Build.cs b/Engine/Plugins/Runtime/Oculus/OculusVR/Source/OculusOpenXRHMD/OculusOpenXRHMD.Build.cs
--- a/Engine/Plugins/Runtime/Oculus/OculusVR/Source/OculusOpenXRHMD/OculusOpenXRHMD.Build.cs
+++ b/Engine/Plugins/Runtime/Oculus/OculusVR/Source/OculusOpenXRHMD/OculusOpenXRHMD.Build.cs
@@ -9,10 +9,7 @@ namespace UnrealBuildTool.Rules
     {
         public OculusOpenXRHMD(ReadOnlyTargetRules Target) : base(Target)
         {
-            PrivateIncludePaths.AddRange(
-                new string[] {
-					Path.Combine(GetModuleDirectory("OpenXRHMD"), "Private"),
-                });
+            PrivateIncludePaths.Add(Path.Combine(GetModuleDirectory("OpenXRHMD"), "Private"));
 
             PublicIncludePathModuleNames.AddRange(
                 new string[] {
diff --git a/Engine/Plugins/Runtime/OpenXR/Source/OpenXRAR/OpenXRAR.Build.cs b/Engine/Plugins/Runtime/OpenXR/Source/OpenXRAR/OpenXRAR.Build.cs
--- a/Engine/Plugins/Runtime/OpenXR/Source/OpenXRAR/OpenXRAR.Build.cs
+++ b/Engine/Plugins/Runtime/OpenXR/Source/OpenXRAR/OpenXRAR.Build.cs
@@ -1,21 +1,14 @@
 // Copyright Epic Games, Inc. All Rights Reserved.
 
-using UnrealBuildTool;
 using System.IO;
+using UnrealBuildTool;
 
 public class OpenXRAR : ModuleRules
 {
     public OpenXRAR(ReadOnlyTargetRules Target) : base(Target)
     {
 		var EngineDir = Path.GetFullPath(Target.RelativeEnginePath);
-		PrivateIncludePaths.AddRange(
-			new string[] {
-					"OpenXRHMD/Private",
-					EngineDir + "/Source/Runtime/Renderer/Private",
-					EngineDir + "/Source/ThirdParty/OpenXR/include",
-				// ... add other private include paths required here ...
-			}
-			);
+		PrivateIncludePaths.Add(Path.Combine(GetModuleDirectory("OpenXRHMD"), "Private"));
 
 		PublicDependencyModuleNames.AddRange(
             new string[]
@@ -23,7 +16,6 @@ public class OpenXRAR : ModuleRules
                 "Core",
                 "Engine",
 				"MRMesh"
-				// ... add other public dependencies that you statically link with here ...
 			}
 			);
 
@@ -38,7 +30,6 @@ public class OpenXRAR : ModuleRules
                 "RHI",
 				"Projects",
 				"OpenXRHMD",
-				// ... add private dependencies that you statically link with here ...
 			}
             );
 			
@@ -52,6 +43,5 @@ public class OpenXRAR : ModuleRules
 		}
 
 		PublicIncludePathModuleNames.Add("OpenXR");
-
 	}
 }
\ No newline at end of file
diff --git a/Engine/Plugins/Runtime/OpenXR/Source/OpenXRHMD/OpenXRHMD.Build.cs b/Engine/Plugins/Runtime/OpenXR/Source/OpenXRHMD/OpenXRHMD.Build.cs
--- a/Engine/Plugins/Runtime/OpenXR/Source/OpenXRHMD/OpenXRHMD.Build.cs
+++ b/Engine/Plugins/Runtime/OpenXR/Source/OpenXRHMD/OpenXRHMD.Build.cs
@@ -1,24 +1,11 @@
 // Copyright Epic Games, Inc. All Rights Reserved.
 
-using System.IO;
-using System.Linq;
-using EpicGames.Core;
-
 namespace UnrealBuildTool.Rules
 {
 	public class OpenXRHMD : ModuleRules
 	{
 		public OpenXRHMD(ReadOnlyTargetRules Target) : base(Target)
         {
-            var EngineDir = Path.GetFullPath(Target.RelativeEnginePath);
-            PrivateIncludePaths.AddRange(
-				new string[] {
-					"OpenXRHMD/Private",
-                    EngineDir + "/Source/ThirdParty/OpenXR/include",
-					// ... add other private include paths required here ...
-				}
-				);
-
 			PublicIncludePathModuleNames.Add("OpenXR");
 
             PublicDependencyModuleNames.Add("HeadMountedDisplay");
diff --git a/Engine/Plugins/Runtime/OpenXR/Source/OpenXRInput/OpenXRInput.Build.cs b/Engine/Plugins/Runtime/OpenXR/Source/OpenXRInput/OpenXRInput.Build.cs
--- a/Engine/Plugins/Runtime/OpenXR/Source/OpenXRInput/OpenXRInput.Build.cs
+++ b/Engine/Plugins/Runtime/OpenXR/Source/OpenXRInput/OpenXRInput.Build.cs
@@ -1,6 +1,5 @@
 // Copyright Epic Games, Inc. All Rights Reserved.
 
-using UnrealBuildTool;
 using System.IO;
 
 namespace UnrealBuildTool.Rules
@@ -9,15 +8,7 @@ namespace UnrealBuildTool.Rules
     {
         public OpenXRInput(ReadOnlyTargetRules Target) : base(Target)
         {
-            var EngineDir = Path.GetFullPath(Target.RelativeEnginePath);
-            PrivateIncludePaths.AddRange(
-                new string[] {
-                    "OpenXRHMD/Private",
-                    EngineDir + "/Source/Runtime/Renderer/Private",
-                    EngineDir + "/Source/ThirdParty/OpenXR/include",
-					// ... add other private include paths required here ...
-				}
-                );
+            PrivateIncludePaths.Add(Path.Combine(GetModuleDirectory("OpenXRHMD"), "Private"));
 
             PrivateDependencyModuleNames.AddRange(
                 new string[]
diff --git a/Engine/Plugins/Runtime/OpenXREyeTracker/Source/OpenXREyeTracker/OpenXREyeTracker.Build.cs b/Engine/Plugins/Runtime/OpenXREyeTracker/Source/OpenXREyeTracker/OpenXREyeTracker.Build.cs
--- a/Engine/Plugins/Runtime/OpenXREyeTracker/Source/OpenXREyeTracker/OpenXREyeTracker.Build.cs
+++ b/Engine/Plugins/Runtime/OpenXREyeTracker/Source/OpenXREyeTracker/OpenXREyeTracker.Build.cs
@@ -1,8 +1,5 @@
 // Copyright Epic Games, Inc. All Rights Reserved.
 
-using UnrealBuildTool;
-using System.IO;
-
 namespace UnrealBuildTool.Rules
 {
     public class OpenXREyeTracker : ModuleRules
@@ -10,19 +7,7 @@ namespace UnrealBuildTool.Rules
         public OpenXREyeTracker(ReadOnlyTargetRules Target) 
 				: base(Target)
         {
-			PublicDependencyModuleNames.AddRange(
-			   new string[]
-			   {
-					"EyeTracker"
-			   }
-		   );
-
-			var EngineDir = Path.GetFullPath(Target.RelativeEnginePath);
-            PrivateIncludePaths.AddRange(
-                new string[] {
-                    EngineDir + "/Source/ThirdParty/OpenXR/include"
-				}
-                );
+			PublicDependencyModuleNames.Add("EyeTracker");
 
             PrivateDependencyModuleNames.AddRange(
                 new string[]
diff --git a/Engine/Plugins/Runtime/OpenXRHandTracking/Source/OpenXRHandTracking/OpenXRHandTracking.Build.cs b/Engine/Plugins/Runtime/OpenXRHandTracking/Source/OpenXRHandTracking/OpenXRHandTracking.Build.cs
--- a/Engine/Plugins/Runtime/OpenXRHandTracking/Source/OpenXRHandTracking/OpenXRHandTracking.Build.cs
+++ b/Engine/Plugins/Runtime/OpenXRHandTracking/Source/OpenXRHandTracking/OpenXRHandTracking.Build.cs
@@ -1,8 +1,5 @@
 // Copyright Epic Games, Inc. All Rights Reserved.
 
-using UnrealBuildTool;
-using System.IO;
-
 namespace UnrealBuildTool.Rules
 {
     public class OpenXRHandTracking: ModuleRules
@@ -10,19 +7,7 @@ namespace UnrealBuildTool.Rules
         public OpenXRHandTracking(ReadOnlyTargetRules Target) 
 				: base(Target)
         {
-			PublicDependencyModuleNames.AddRange(
-			   new string[]
-			   {
-					"InputDevice"
-			   }
-		   );
-
-			var EngineDir = Path.GetFullPath(Target.RelativeEnginePath);
-            PrivateIncludePaths.AddRange(
-                new string[] {
-                    EngineDir + "/Source/ThirdParty/OpenXR/include"
-				}
-                );
+			PublicDependencyModuleNames.Add("InputDevice");
 
             PrivateDependencyModuleNames.AddRange(
                 new string[]
