Fixed sorting function that made world locked layers render on top of face locked layers.

diff --git a/Engine/Plugins/Runtime/OpenXR/Source/OpenXRHMD/Private/OpenXRHMD.cpp b/Engine/Plugins/Runtime/OpenXR/Source/OpenXRHMD/Private/OpenXRHMD.cpp
--- a/Engine/Plugins/Runtime/OpenXR/Source/OpenXRHMD/Private/OpenXRHMD.cpp
+++ b/Engine/Plugins/Runtime/OpenXR/Source/OpenXRHMD/Private/OpenXRHMD.cpp
@@ -2453,6 +2453,12 @@ void FOpenXRHMD::SetupFrameQuadLayers_RenderThread(FRHICommandListImmediate& RHI
 			FLayerDesc DescA, DescB;
 			if (GetLayerDescMember(A, DescA) && GetLayerDescMember(B, DescB))
 			{
+				bool bOneIsFaceLocked = (DescA.PositionType == IStereoLayers::FaceLocked) || (DescB.PositionType == IStereoLayers::FaceLocked);
+				bool bBothFaceLocked = DescA.PositionType == IStereoLayers::FaceLocked && DescB.PositionType == IStereoLayers::FaceLocked;
+				if (bOneIsFaceLocked && !bBothFaceLocked)
+				{
+					return DescB.PositionType == IStereoLayers::FaceLocked;
+				}
 				if (DescA.Priority < DescB.Priority)
 				{
 					return true;
