Deprecate GetViewConfigurations and GetViewLocations in IOpenXRExtensionPlugin

diff --git a/Engine/Plugins/Runtime/OpenXR/Source/OpenXRHMD/Private/OpenXRHMD.cpp b/Engine/Plugins/Runtime/OpenXR/Source/OpenXRHMD/Private/OpenXRHMD.cpp
--- a/Engine/Plugins/Runtime/OpenXR/Source/OpenXRHMD/Private/OpenXRHMD.cpp
+++ b/Engine/Plugins/Runtime/OpenXR/Source/OpenXRHMD/Private/OpenXRHMD.cpp
@@ -1463,7 +1463,9 @@ void FOpenXRHMD::EnumerateViews(FPipelinedFrameState& PipelineState)
 	for (IOpenXRExtensionPlugin* Module : ExtensionPlugins)
 	{
 		TArray<XrViewConfigurationView> ViewConfigs;
+		PRAGMA_DISABLE_DEPRECATION_WARNINGS;
 		Module->GetViewConfigurations(System, ViewConfigs);
+		PRAGMA_ENABLE_DEPRECATION_WARNINGS;
 
 		const EStereoscopicPass PluginPassType = ViewConfigs.Num() > 1 ? EStereoscopicPass::eSSP_PRIMARY : EStereoscopicPass::eSSP_FULL;
 		for (int32 i = 0; i < ViewConfigs.Num(); i++)
@@ -1490,7 +1492,9 @@ void FOpenXRHMD::EnumerateViews(FPipelinedFrameState& PipelineState)
 			if (PipelineState.PluginViewInfos.ContainsByPredicate(Predicate))
 			{
 				TArray<XrView> Views;
+				PRAGMA_DISABLE_DEPRECATION_WARNINGS;
 				Module->GetViewLocations(Session, PipelineState.FrameState.predictedDisplayTime, DeviceSpaces[HMDDeviceId].Space, Views);
+				PRAGMA_ENABLE_DEPRECATION_WARNINGS;
 				check(Views.Num() > 0);
 				PipelineState.Views.Append(Views);
 			}
diff --git a/Engine/Plugins/Runtime/OpenXR/Source/OpenXRHMD/Public/IOpenXRExtensionPlugin.h b/Engine/Plugins/Runtime/OpenXR/Source/OpenXRHMD/Public/IOpenXRExtensionPlugin.h
--- a/Engine/Plugins/Runtime/OpenXR/Source/OpenXRHMD/Public/IOpenXRExtensionPlugin.h
+++ b/Engine/Plugins/Runtime/OpenXR/Source/OpenXRHMD/Public/IOpenXRExtensionPlugin.h
@@ -290,6 +290,7 @@ public:
 	/**
 	* Callback to provide extra view configurations that should be rendered in the main render pass
 	*/
+	UE_DEPRECATED(5.2, "Support for new view configurations should be added to FOpenXRHMD instead.")
 	virtual void GetViewConfigurations(XrSystemId InSystem, TArray<XrViewConfigurationView>& OutViews)
 	{
 	}
@@ -297,6 +298,7 @@ public:
 	/**
 	* Callback to provide the pose and fov of each view that was provided in GetViewConfigurations
 	*/
+	UE_DEPRECATED(5.2, "Support for new view configurations should be added to FOpenXRHMD instead.")
 	virtual void GetViewLocations(XrSession InSession, XrTime InDisplayTime, XrSpace InViewSpace, TArray<XrView>& OutViews)
 	{
 	}
