UE-177875 In 5.2 QAGame, OpenXR-Loader causes a hang during initialization
-Switched the default on the cvar to true because the current openxr loader hangs during initialization on some devices rather than failing.
-Also renamed it from vr. to the preferred xr.

diff --git a/Engine/Plugins/Runtime/OpenXR/Source/OpenXRHMD/Private/OpenXRHMDModule.cpp b/Engine/Plugins/Runtime/OpenXR/Source/OpenXRHMD/Private/OpenXRHMDModule.cpp
--- a/Engine/Plugins/Runtime/OpenXR/Source/OpenXRHMD/Private/OpenXRHMDModule.cpp
+++ b/Engine/Plugins/Runtime/OpenXR/Source/OpenXRHMD/Private/OpenXRHMDModule.cpp
@@ -26,9 +26,9 @@ static TAutoConsoleVariable<int32> CVarEnableOpenXRValidationLayer(
 	ECVF_Default);		// @todo: Should we specify ECVF_Cheat here so this doesn't show up in release builds?
 
 static TAutoConsoleVariable<bool> CVarDisableOpenXROnAndroidWithoutOculus(
-	TEXT("vr.DisableOpenXROnAndroidWithoutOculus"),
-	false,
-	TEXT("If true OpenXR will not initialize on Android unless the project is packaged for OpenXR.  (Potentially useful for cross platform projects invovling openxr and arcore, though this is achievable in other ways.)"),
+	TEXT("xr.DisableOpenXROnAndroidWithoutOculus"),
+	true,
+	TEXT("If true OpenXR will not initialize on Android unless the project is packaged for Oculus (ProjectSetting->Platforms->Android->Advanced APK Packaging->PackageForOculusMobileDevices list not empty).  Currently defaulted to true because the OpenXR loader we are using hangs during intialization on some devices instead of failing, as it should."),
 	ECVF_RenderThreadSafe);
 
 //---------------------------------------------------
