Implement legacy button repeat behaviour

diff --git a/Engine/Plugins/Runtime/OpenXR/Source/OpenXRInput/Private/OpenXRInput.cpp b/Engine/Plugins/Runtime/OpenXR/Source/OpenXRInput/Private/OpenXRInput.cpp
--- a/Engine/Plugins/Runtime/OpenXR/Source/OpenXRInput/Private/OpenXRInput.cpp
+++ b/Engine/Plugins/Runtime/OpenXR/Source/OpenXRInput/Private/OpenXRInput.cpp
@@ -851,12 +851,20 @@ void FOpenXRInputPlugin::FOpenXRInput::SendControllerEvents()
 						MessageHandler->OnControllerButtonReleased(Action.Name, DeviceMapper.GetPrimaryPlatformUser(), DeviceMapper.GetDefaultInputDevice(), /*IsRepeat =*/false);
 					}
 
+					Action.NextRepeatTime = State.lastChangeTime + InitialButtonRepeatDelay;
+
 					FXRTimedInputActionDelegate* const Delegate = OpenXRInputNamespace::GetTimedInputActionDelegate(Action.Name);
 					if (Delegate)
 					{
 						Delegate->Execute(State.currentState ? 1.0 : 0.0f, ToFTimespan(State.lastChangeTime));
 					}
 				}
+				else if (State.currentState && OpenXRHMD->GetDisplayTime() >= Action.NextRepeatTime)
+				{
+					// TODO: We should retrieve the current time rather than the display time
+					MessageHandler->OnControllerButtonPressed(Action.Name, DeviceMapper.GetPrimaryPlatformUser(), DeviceMapper.GetDefaultInputDevice(), /*IsRepeat =*/true);
+					Action.NextRepeatTime = OpenXRHMD->GetDisplayTime() + ButtonRepeatDelay;
+				}
 			}
 			break;
 			case XR_ACTION_TYPE_FLOAT_INPUT:
diff --git a/Engine/Plugins/Runtime/OpenXR/Source/OpenXRInput/Private/OpenXRInput.h b/Engine/Plugins/Runtime/OpenXR/Source/OpenXRInput/Private/OpenXRInput.h
--- a/Engine/Plugins/Runtime/OpenXR/Source/OpenXRInput/Private/OpenXRInput.h
+++ b/Engine/Plugins/Runtime/OpenXR/Source/OpenXRInput/Private/OpenXRInput.h
@@ -31,6 +31,9 @@ public:
 		FName			Name;
 		XrAction		Handle;
 
+		// Legacy Input
+		XrTime			NextRepeatTime;
+
 		// Enhanced Input
 		TObjectPtr<const UInputAction> Object;
 		TMultiMap<TPair<XrPath, XrPath>, TObjectPtr<UInputTrigger>> Triggers;
@@ -159,6 +162,10 @@ public:
 		*/
 		float CurrentDeltaTime = 0.0f;
 
+		/** Repeat key delays */
+		const XrTime InitialButtonRepeatDelay = 2e8;
+		const XrTime ButtonRepeatDelay = 1e8;
+
 		bool BuildActions(XrSession Session);
 		void SyncActions(XrSession Session);
 		void BuildLegacyActions(TMap<FString, FInteractionProfile>& Profiles);
