Added support for cylinder and equirect layers in openXR plugin + refactor of layer management code to make it agnostic to the type of layer.

Reworked from PR #9368: Support Cylinder and Equirect OpenXR Layers (Contributed by [REDACTED_PERSON]).

diff --git a/Engine/Plugins/Runtime/OpenXR/Source/OpenXRHMD/Private/OpenXRHMD.cpp b/Engine/Plugins/Runtime/OpenXR/Source/OpenXRHMD/Private/OpenXRHMD.cpp
--- a/Engine/Plugins/Runtime/OpenXR/Source/OpenXRHMD/Private/OpenXRHMD.cpp
+++ b/Engine/Plugins/Runtime/OpenXR/Source/OpenXRHMD/Private/OpenXRHMD.cpp
@@ -1954,12 +1954,12 @@ void FOpenXRHMD::DestroySession()
 			Layer.LeftEye.Swapchain.Reset();
 		});
 
-		NativeQuadLayers.Reset();
+		NativeLayers.Reset();
 		EmulatedFaceLockedLayers.Reset();
 
 		PipelinedLayerStateRendering.ColorSwapchain.Reset();
 		PipelinedLayerStateRendering.DepthSwapchain.Reset();
-		PipelinedLayerStateRendering.QuadSwapchains.Reset();
+		PipelinedLayerStateRendering.NativeOverlaySwapchains.Reset();
 		PipelinedLayerStateRendering.EmulatedLayerState.EmulationSwapchain.Reset();
 
 		// TODO: Once we handle OnFinishRendering_RHIThread + StopSession interactions
@@ -1967,7 +1967,7 @@ void FOpenXRHMD::DestroySession()
 		// `ensure` here to make sure these are released.
 		PipelinedLayerStateRHI.ColorSwapchain.Reset();
 		PipelinedLayerStateRHI.DepthSwapchain.Reset();
-		PipelinedLayerStateRHI.QuadSwapchains.Reset();
+		PipelinedLayerStateRHI.NativeOverlaySwapchains.Reset();
 		PipelinedLayerStateRHI.EmulatedLayerState.EmulationSwapchain.Reset();
 
 		PipelinedFrameStateGame.TrackingSpace.Reset();
@@ -2373,35 +2373,30 @@ bool FOpenXRHMD::IsEmulatingStereoLayers()
 	return !bLayerSupportOpenXRCompliant || CVarOpenXRForceStereoLayerEmulation.GetValueOnAnyThread();
 }
 
-void FOpenXRHMD::SetupFrameQuadLayers_RenderThread(FRHICommandListImmediate& RHICmdList)
+void FOpenXRHMD::SetupFrameLayers_RenderThread(FRHICommandListImmediate& RHICmdList)
 {
 	ensure(IsInRenderingThread());
-
 	if (GetStereoLayersDirty())
 	{
-		// When SetupFrameQuadLayers_RenderThread is called more than once with GetStereoLayersDirty = true,
-		// NativeQuadLayers is rebuilt from the internal array of stereo layers. However, the internal array is
+		// When SetupFrameLayers_RenderThread is called more than once with GetStereoLayersDirty = true,
+		// NativeLayers is rebuilt from the internal array of stereo layers. However, the internal array is
 		// not updated after a static swapchain is created for a layer and it always retains bUpdateTexture = true
 		// which leads to the swapchain trying to acquire a second image and resulting in a crash.
 		// This serves as a workaround that updates the internal state of the layers mirroring the most recent state
 		// of the NativeQuadLayers.
-		TArray<FOpenXRLayer> NativeQuadLayersBackup = NativeQuadLayers;
-		// Go over the dirtied layers to bin them into either native or emulated
+		TArray<FOpenXRLayer> NativeLayersBackup = NativeLayers;
+		
 		BackgroundCompositedEmulatedLayers.Reset();
 		EmulatedFaceLockedLayers.Reset();
-		NativeQuadLayers.Reset();
+		NativeLayers.Reset();
 
+		// Go over the dirtied layers to bin them into either native or emulated
 		ForEachLayer([&](uint32 /* unused */, FOpenXRLayer& Layer)
 		{
-			// OpenXR currently supports only Quad layers
-			if (!Layer.Desc.HasShape<FQuadLayer>())
-			{
-				return;
-			}
 			if (IsEmulatingStereoLayers())
 			{
-
-				if (!Layer.Desc.IsVisible())
+				// Only quad layers are supported by emulation.
+				if (!Layer.Desc.HasShape<FQuadLayer>() || !Layer.Desc.IsVisible())
 				{
 					return;
 				}
@@ -2424,30 +2419,7 @@ void FOpenXRHMD::SetupFrameQuadLayers_RenderThread(FRHICommandListImmediate& RHI
 			} // OpenXR compliant layer support (16 layers).
 			else 
 			{
-				if (Layer.Desc.IsVisible())
-				{
-					CreateNativeLayerSwapchain(Layer, RenderBridge, Session);
-					FOpenXRLayer& LastLayer = NativeQuadLayers.Add_GetRef(Layer);
-					FOpenXRLayer* FoundLayer = NativeQuadLayersBackup.FindByPredicate([LayerId = LastLayer.GetLayerId()](const FOpenXRLayer& Layer)
-					{
-						if (Layer.GetLayerId() == LayerId)
-						{
-							return true;
-						}
-						return false;
-					});
-					if (FoundLayer != nullptr)
-					{
-						LastLayer.RightEye.bUpdateTexture = FoundLayer->RightEye.bUpdateTexture;
-						LastLayer.LeftEye.bUpdateTexture = FoundLayer->LeftEye.bUpdateTexture;
-					}
-				}
-				else
-				{
-					// We retain references in FPipelinedLayerState to avoid premature destruction
-					Layer.RightEye.Swapchain.Reset();
-					Layer.LeftEye.Swapchain.Reset();
-				}
+				ConfigureLayerSwapchain(Layer, NativeLayersBackup);
 			}
 		});
 
@@ -2477,7 +2449,7 @@ void FOpenXRHMD::SetupFrameQuadLayers_RenderThread(FRHICommandListImmediate& RHI
 
 		BackgroundCompositedEmulatedLayers.Sort(LayerCompare);
 		EmulatedFaceLockedLayers.Sort(LayerCompare);
-		NativeQuadLayers.Sort(LayerCompare);
+		NativeLayers.Sort(LayerCompare);
 
 		PipelinedLayerStateRendering.EmulatedLayerState.bIsFaceLockedLayerEmulationActive = !EmulatedFaceLockedLayers.IsEmpty();
 	}
@@ -2485,67 +2457,82 @@ void FOpenXRHMD::SetupFrameQuadLayers_RenderThread(FRHICommandListImmediate& RHI
 	const FTransform InvTrackingToWorld = GetTrackingToWorldTransform().Inverse();
 	const float WorldToMeters = GetWorldToMetersScale();
 
-	// Set up our OpenXR info per native quad layer. Emulated layers have everything in FLayerDesc
-	PipelinedLayerStateRendering.QuadLayers.Reset(NativeQuadLayers.Num());
-	PipelinedLayerStateRendering.QuadSwapchains.Reset(NativeQuadLayers.Num());
+	PipelinedLayerStateRendering.NativeOverlays.Reset(NativeLayers.Num());
+	PipelinedLayerStateRendering.NativeOverlaySwapchains.Reset(NativeLayers.Num());
+
+	// Set up our OpenXR info per native layer. Emulated layers have everything in FLayerDesc.
+	for (const FOpenXRLayer& Layer : NativeLayers)
 	{
 		FReadScopeLock DeviceLock(DeviceMutex);
 
-		for (const FOpenXRLayer& Layer : NativeQuadLayers)
-		{
-			const bool bNoAlpha = Layer.Desc.Flags & IStereoLayers::LAYER_FLAG_TEX_NO_ALPHA_CHANNEL;
-			const bool bIsStereo = Layer.Desc.LeftTexture.IsValid();
-			const FTransform PositionTransform = Layer.Desc.PositionType == ELayerType::WorldLocked ?
-				InvTrackingToWorld : FTransform::Identity;
-
-			XrCompositionLayerQuad Quad = { XR_TYPE_COMPOSITION_LAYER_QUAD, nullptr };
-			Quad.layerFlags = bNoAlpha ? 0 : XR_COMPOSITION_LAYER_UNPREMULTIPLIED_ALPHA_BIT |
-				XR_COMPOSITION_LAYER_BLEND_TEXTURE_SOURCE_ALPHA_BIT;
-			Quad.space = Layer.Desc.PositionType == ELayerType::FaceLocked ?
-				DeviceSpaces[HMDDeviceId].Space : PipelinedFrameStateRendering.TrackingSpace->Handle;
-			Quad.subImage.imageArrayIndex = 0;
-			Quad.pose = ToXrPose(Layer.Desc.Transform * PositionTransform, WorldToMeters);
-
-			// The layer pose doesn't take the transform scale into consideration, so we need to manually apply it the quad size.
-			const FVector2D LayerComponentScaler(Layer.Desc.Transform.GetScale3D().Y, Layer.Desc.Transform.GetScale3D().Z);
-			const ETextureCopyBlendModifier SrcTextureCopyModifier = bNoAlpha ? ETextureCopyBlendModifier::Opaque : ETextureCopyBlendModifier::TransparentAlphaPassthrough;
+		XrSpace Space = Layer.Desc.PositionType == ELayerType::FaceLocked ?
+			DeviceSpaces[HMDDeviceId].Space : PipelinedFrameStateRendering.TrackingSpace->Handle;
+		
+		TArray<FXrCompositionLayerUnion> Headers = Layer.CreateOpenXRLayer(InvTrackingToWorld, WorldToMeters, Space);
+		PipelinedLayerStateRendering.NativeOverlays.Append(Headers);
+		UpdateLayerSwapchainTexture(Layer, RHICmdList);
+	}
+}
 
-			// We need to copy each layer into an OpenXR swapchain so they can be displayed by the compositor
-			if (Layer.RightEye.Swapchain.IsValid() && Layer.Desc.Texture.IsValid())
+void FOpenXRHMD::ConfigureLayerSwapchain(FOpenXRLayer& Layer, TArray<FOpenXRLayer>& BackupLayers)
+{
+	// OpenXR currently supports only Quad layers unless the cylinder and equirect extensions are enabled.
+	if ((Layer.Desc.HasShape<FCylinderLayer>() && IsExtensionEnabled(XR_KHR_COMPOSITION_LAYER_CYLINDER_EXTENSION_NAME)) ||
+		(Layer.Desc.HasShape<FEquirectLayer>() && IsExtensionEnabled(XR_KHR_COMPOSITION_LAYER_EQUIRECT_EXTENSION_NAME)) ||
+		Layer.Desc.HasShape<FQuadLayer>())
+	{
+		if (Layer.Desc.IsVisible())
+		{
+			CreateNativeLayerSwapchain(Layer, RenderBridge, Session);
+			FOpenXRLayer& LastLayer = NativeLayers.Add_GetRef(Layer);
+			FOpenXRLayer* FoundLayer = BackupLayers.FindByPredicate([LayerId = LastLayer.GetLayerId()](const FOpenXRLayer& Layer)
 			{
-				if (Layer.RightEye.bUpdateTexture && bIsRunning)
+				if (Layer.GetLayerId() == LayerId)
 				{
-					FRHITexture2D* SrcTexture = Layer.Desc.Texture->GetTexture2D();
-					const FIntRect DstRect(FIntPoint(0, 0), Layer.RightEye.SwapchainSize.IntPoint());
-					CopyTexture_RenderThread(RHICmdList, SrcTexture, FIntRect(), Layer.RightEye.Swapchain, DstRect, false, SrcTextureCopyModifier);
+					return true;
 				}
-
-				Quad.eyeVisibility = bIsStereo ? XR_EYE_VISIBILITY_RIGHT : XR_EYE_VISIBILITY_BOTH;
-
-				Quad.subImage.imageRect = ToXrRect(Layer.GetRightViewportSize());
-				Quad.subImage.swapchain = static_cast<FOpenXRSwapchain*>(Layer.RightEye.Swapchain.Get())->GetHandle();
-				Quad.size = ToXrExtent2D(Layer.GetRightQuadSize() * LayerComponentScaler, WorldToMeters);
-				PipelinedLayerStateRendering.QuadLayers.Add(Quad);
-				PipelinedLayerStateRendering.QuadSwapchains.Add(Layer.RightEye.Swapchain);
+				return false;
+			});
+			if (FoundLayer != nullptr)
+			{
+				LastLayer.RightEye.bUpdateTexture = FoundLayer->RightEye.bUpdateTexture;
+				LastLayer.LeftEye.bUpdateTexture = FoundLayer->LeftEye.bUpdateTexture;
 			}
+		}
+		else
+		{
+			// We retain references in FPipelinedLayerState to avoid premature destruction
+			Layer.RightEye.Swapchain.Reset();
+			Layer.LeftEye.Swapchain.Reset();
+		}
+	}
+}
 
-			if (Layer.LeftEye.Swapchain.IsValid() && Layer.Desc.LeftTexture.IsValid())
-			{
-				if (Layer.LeftEye.bUpdateTexture && bIsRunning)
-				{
-					FRHITexture2D* SrcTexture = Layer.Desc.LeftTexture->GetTexture2D();
-					const FIntRect DstRect(FIntPoint(0, 0), Layer.LeftEye.SwapchainSize.IntPoint());
-					CopyTexture_RenderThread(RHICmdList, SrcTexture, FIntRect(), Layer.LeftEye.Swapchain, DstRect, false, SrcTextureCopyModifier);
-				}
+void FOpenXRHMD::UpdateLayerSwapchainTexture(const FOpenXRLayer& Layer, FRHICommandListImmediate& RHICmdList)
+{
+	const bool bNoAlpha = Layer.Desc.Flags & IStereoLayers::LAYER_FLAG_TEX_NO_ALPHA_CHANNEL;
+	const ETextureCopyBlendModifier SrcTextureCopyModifier = bNoAlpha ? ETextureCopyBlendModifier::Opaque : ETextureCopyBlendModifier::TransparentAlphaPassthrough;
 
-				Quad.eyeVisibility = XR_EYE_VISIBILITY_LEFT;
-				Quad.subImage.imageRect = ToXrRect(Layer.GetLeftViewportSize());
-				Quad.subImage.swapchain = static_cast<FOpenXRSwapchain*>(Layer.LeftEye.Swapchain.Get())->GetHandle();
-				Quad.size = ToXrExtent2D(Layer.GetLeftQuadSize() * LayerComponentScaler, WorldToMeters);
-				PipelinedLayerStateRendering.QuadLayers.Add(Quad);
-				PipelinedLayerStateRendering.QuadSwapchains.Add(Layer.LeftEye.Swapchain);
-			}
+	// We need to copy each layer into an OpenXR swapchain so they can be displayed by the compositor.
+	if (Layer.RightEye.Swapchain.IsValid() && Layer.Desc.Texture.IsValid())
+	{
+		if (Layer.RightEye.bUpdateTexture && bIsRunning)
+		{
+			FRHITexture2D* SrcTexture = Layer.Desc.Texture->GetTexture2D();
+			FIntRect DstRect(FIntPoint(0, 0), Layer.RightEye.SwapchainSize.IntPoint());
+			CopyTexture_RenderThread(RHICmdList, SrcTexture, FIntRect(), Layer.RightEye.Swapchain, DstRect, false, SrcTextureCopyModifier);
 		}
+		PipelinedLayerStateRendering.NativeOverlaySwapchains.Add(Layer.RightEye.Swapchain);
+	}
+	if (Layer.LeftEye.Swapchain.IsValid() && Layer.Desc.LeftTexture.IsValid())
+	{
+		if (Layer.LeftEye.bUpdateTexture && bIsRunning)
+		{
+			FRHITexture2D* SrcTexture = Layer.Desc.LeftTexture->GetTexture2D();
+			FIntRect DstRect(FIntPoint(0, 0), Layer.LeftEye.SwapchainSize.IntPoint());
+			CopyTexture_RenderThread(RHICmdList, SrcTexture, FIntRect(), Layer.LeftEye.Swapchain, DstRect, false, SrcTextureCopyModifier);
+		}
+		PipelinedLayerStateRendering.NativeOverlaySwapchains.Add(Layer.LeftEye.Swapchain);
 	}
 }
 
@@ -2706,7 +2693,7 @@ void FOpenXRHMD::OnBeginRendering_RenderThread(FRHICommandListImmediate& RHICmdL
 		}
 	}
 
-	SetupFrameQuadLayers_RenderThread(RHICmdList);
+	SetupFrameLayers_RenderThread(RHICmdList);
 
 #if !PLATFORM_HOLOLENS
 	if (bHiddenAreaMaskSupported && bNeedReBuildOcclusionMesh)
@@ -2725,11 +2712,11 @@ void FOpenXRHMD::OnBeginRendering_RenderThread(FRHICommandListImmediate& RHICmdL
 		SCOPED_NAMED_EVENT(EnqueueFrame, FColor::Red);
 
 		// Reset the update flag on native quad layers
-		for (FOpenXRLayer& NativeQuadLayer : NativeQuadLayers)
+		for (FOpenXRLayer& Layer : NativeLayers)
 		{
-			const bool bUpdateTexture = NativeQuadLayer.Desc.Flags & IStereoLayers::LAYER_FLAG_TEX_CONTINUOUS_UPDATE;
-			NativeQuadLayer.RightEye.bUpdateTexture = bUpdateTexture;
-			NativeQuadLayer.LeftEye.bUpdateTexture = bUpdateTexture;
+			const bool bUpdateTexture = Layer.Desc.Flags & IStereoLayers::LAYER_FLAG_TEX_CONTINUOUS_UPDATE;
+			Layer.RightEye.bUpdateTexture = bUpdateTexture;
+			Layer.LeftEye.bUpdateTexture = bUpdateTexture;
 		}
 
 		FXRSwapChainPtr ColorSwapchain = PipelinedLayerStateRendering.ColorSwapchain;
@@ -3243,10 +3230,7 @@ void FOpenXRHMD::OnFinishRendering_RHIThread()
 			Headers.Add(reinterpret_cast<const XrCompositionLayerBaseHeader*>(&CompositedLayer));
 		}
 
-		for (const XrCompositionLayerQuad& Quad : PipelinedLayerStateRHI.QuadLayers)
-		{
-			Headers.Add(reinterpret_cast<const XrCompositionLayerBaseHeader*>(&Quad));
-		}
+		AddLayersToHeaders(Headers);
 
 		int32 BlendModeOverride = CVarOpenXREnvironmentBlendMode.GetValueOnRenderThread();
 
@@ -3286,6 +3270,14 @@ void FOpenXRHMD::OnFinishRendering_RHIThread()
 	bIsRendering = false;
 }
 
+void FOpenXRHMD::AddLayersToHeaders(TArray<const XrCompositionLayerBaseHeader*>& Headers)
+{
+	for (const FXrCompositionLayerUnion& Layer : PipelinedLayerStateRHI.NativeOverlays)
+	{
+		Headers.Add(reinterpret_cast<const XrCompositionLayerBaseHeader*>(&Layer.Header));
+	}
+}
+
 FXRRenderBridge* FOpenXRHMD::GetActiveRenderBridge_GameThread(bool /* bUseSeparateRenderTarget */)
 {
 	return RenderBridge;
@@ -3780,7 +3772,7 @@ void FOpenXRHMD::DrawVisibleAreaMesh(class FRHICommandList& RHICmdList, int32 Vi
 
 void FOpenXRHMD::UpdateLayer(FOpenXRLayer& ManagerLayer, uint32 LayerId, bool bIsValid)
 {
-	for (FOpenXRLayer& NativeLayer : NativeQuadLayers)
+	for (FOpenXRLayer& NativeLayer : NativeLayers)
 	{
 		if (NativeLayer.GetLayerId() == LayerId)
 		{
diff --git a/Engine/Plugins/Runtime/OpenXR/Source/OpenXRHMD/Private/OpenXRHMD.h b/Engine/Plugins/Runtime/OpenXR/Source/OpenXRHMD/Private/OpenXRHMD.h
--- a/Engine/Plugins/Runtime/OpenXR/Source/OpenXRHMD/Private/OpenXRHMD.h
+++ b/Engine/Plugins/Runtime/OpenXR/Source/OpenXRHMD/Private/OpenXRHMD.h
@@ -25,6 +25,7 @@ class UCanvas;
 class FOpenXRRenderBridge;
 class IOpenXRInputModule;
 struct FDefaultStereoLayers_LayerRenderParams;
+union FXrCompositionLayerUnion;
 
 /**
  * Simple Head Mounted Display
@@ -107,7 +108,7 @@ public:
 
 	struct FPipelinedLayerState
 	{
-		TArray<XrCompositionLayerQuad> QuadLayers;
+		TArray<FXrCompositionLayerUnion> NativeOverlays;
 		TArray<XrCompositionLayerProjectionView> ProjectionLayers;
 		TArray<XrCompositionLayerDepthInfoKHR> DepthLayers;
 
@@ -116,7 +117,7 @@ public:
 
 		FXRSwapChainPtr ColorSwapchain;
 		FXRSwapChainPtr DepthSwapchain;
-		TArray<FXRSwapChainPtr> QuadSwapchains;
+		TArray<FXRSwapChainPtr> NativeOverlaySwapchains;
 
 		FEmulatedLayerState EmulatedLayerState;
 
@@ -233,7 +234,6 @@ public:
 	virtual bool HDRGetMetaDataForStereo(EDisplayOutputFormat& OutDisplayOutputFormat, EDisplayColorGamut& OutDisplayColorGamut, bool& OutbHDRSupported) override;
 
 protected:
-
 	enum ETextureCopyBlendModifier : uint8;
 
 	bool StartSession();
@@ -266,7 +266,7 @@ protected:
 	// Used with FCoreDelegates
 	void VRHeadsetRecenterDelegate();
 
-	void SetupFrameQuadLayers_RenderThread(FRHICommandListImmediate& RHICmdList);
+	void SetupFrameLayers_RenderThread(FRHICommandListImmediate& RHICmdList);
 	void DrawEmulatedLayers_RenderThread(FRDGBuilder& GraphBuilder, const FSceneView& InView);
 	void DrawBackgroundCompositedEmulatedLayers_RenderThread(FRDGBuilder& GraphBuilder, const FSceneView& InView);
 	void DrawEmulatedFaceLockedLayers_RenderThread(FRDGBuilder& GraphBuilder, const FSceneView& InView);
@@ -393,10 +393,15 @@ public:
 	OPENXRHMD_API void SetEnvironmentBlendMode(XrEnvironmentBlendMode NewBlendMode);
 
 private:
+
 	TArray<XrEnvironmentBlendMode> RetrieveEnvironmentBlendModes() const;
 	FDefaultStereoLayers_LayerRenderParams CalculateEmulatedLayerRenderParams(const FSceneView& InView);
 	FRHIRenderPassInfo SetupEmulatedLayersRenderPass(FRHICommandListImmediate& RHICmdList, const FSceneView& InView, TArray<IStereoLayers::FLayerDesc>& Layers, FTexture2DRHIRef RenderTarget, FDefaultStereoLayers_LayerRenderParams& OutRenderParams);
 	bool IsEmulatingStereoLayers();
+	
+	void UpdateLayerSwapchainTexture(const FOpenXRLayer& Layer, FRHICommandListImmediate& RHICmdList);
+	void ConfigureLayerSwapchain(FOpenXRLayer& Layer, TArray<FOpenXRLayer>& BackupLayers);
+	void AddLayersToHeaders(TArray<const XrCompositionLayerBaseHeader*>& Headers);
 
 	bool					bStereoEnabled;
 	TAtomic<bool>			bIsRunning;
@@ -464,5 +469,5 @@ private:
 	bool					bLayerSupportOpenXRCompliant;
 	TArray<IStereoLayers::FLayerDesc> BackgroundCompositedEmulatedLayers;
 	TArray<IStereoLayers::FLayerDesc> EmulatedFaceLockedLayers;
-	TArray<FOpenXRLayer>			  NativeQuadLayers;
+	TArray<FOpenXRLayer>			  NativeLayers;
 };
diff --git a/Engine/Plugins/Runtime/OpenXR/Source/OpenXRHMD/Private/OpenXRHMDModule.cpp b/Engine/Plugins/Runtime/OpenXR/Source/OpenXRHMD/Private/OpenXRHMDModule.cpp
--- a/Engine/Plugins/Runtime/OpenXR/Source/OpenXRHMD/Private/OpenXRHMDModule.cpp
+++ b/Engine/Plugins/Runtime/OpenXR/Source/OpenXRHMD/Private/OpenXRHMDModule.cpp
@@ -492,6 +492,8 @@ bool FOpenXRHMDModule::GetOptionalExtensions(TArray<const ANSICHAR*>& OutExtensi
 	OutExtensions.Add(XR_KHR_VULKAN_SWAPCHAIN_FORMAT_LIST_EXTENSION_NAME);
 #endif
 	OutExtensions.Add(XR_KHR_COMPOSITION_LAYER_DEPTH_EXTENSION_NAME);
+	OutExtensions.Add(XR_KHR_COMPOSITION_LAYER_CYLINDER_EXTENSION_NAME);
+	OutExtensions.Add(XR_KHR_COMPOSITION_LAYER_EQUIRECT_EXTENSION_NAME);
 	OutExtensions.Add(XR_VARJO_QUAD_VIEWS_EXTENSION_NAME);
 	OutExtensions.Add(XR_KHR_VISIBILITY_MASK_EXTENSION_NAME);
 	OutExtensions.Add(XR_KHR_BINDING_MODIFICATION_EXTENSION_NAME);
diff --git a/Engine/Plugins/Runtime/OpenXR/Source/OpenXRHMD/Private/OpenXRHMD_Layer.cpp b/Engine/Plugins/Runtime/OpenXR/Source/OpenXRHMD/Private/OpenXRHMD_Layer.cpp
--- a/Engine/Plugins/Runtime/OpenXR/Source/OpenXRHMD/Private/OpenXRHMD_Layer.cpp
+++ b/Engine/Plugins/Runtime/OpenXR/Source/OpenXRHMD/Private/OpenXRHMD_Layer.cpp
@@ -1,6 +1,8 @@
 // Copyright Epic Games, Inc. All Rights Reserved.
 
 #include "OpenXRHMD_Layer.h"
+#include "IStereoLayers.h"
+#include "OpenXRHMD_Swapchain.h"
 #include "OpenXRCore.h"
 #include "OpenXRPlatformRHI.h"
 
@@ -77,6 +79,135 @@ FVector2D FOpenXRLayer::GetLeftQuadSize() const
 	}
 	return Desc.QuadSize;
 }
+TArray<FXrCompositionLayerUnion> FOpenXRLayer::CreateOpenXRLayer(FTransform InvTrackingToWorld, float WorldToMeters, XrSpace Space) const
+{
+	TArray<FXrCompositionLayerUnion> Headers;
+
+	const bool bNoAlpha = Desc.Flags & IStereoLayers::LAYER_FLAG_TEX_NO_ALPHA_CHANNEL;
+	const bool bIsStereo = Desc.LeftTexture.IsValid();
+	FTransform PositionTransform = Desc.PositionType == IStereoLayers::ELayerType::WorldLocked ?
+		InvTrackingToWorld : FTransform::Identity;
+
+	if (Desc.HasShape<FQuadLayer>())
+	{
+		CreateOpenXRQuadLayer(bIsStereo, bNoAlpha, InvTrackingToWorld, WorldToMeters, Space, Headers);
+	}
+	else if (Desc.HasShape<FCylinderLayer>())
+	{
+		CreateOpenXRCylinderLayer(bIsStereo, bNoAlpha, InvTrackingToWorld, WorldToMeters, Space, Headers);
+	}
+	else if (Desc.HasShape<FEquirectLayer>())
+	{
+		CreateOpenXREquirectLayer(bIsStereo, bNoAlpha, InvTrackingToWorld, WorldToMeters, Space, Headers);
+	}
+
+	return Headers;
+}
+
+void FOpenXRLayer::CreateOpenXRCylinderLayer(bool bIsStereo, bool bNoAlpha, FTransform PositionTransform, float WorldToMeters, XrSpace Space, TArray<FXrCompositionLayerUnion>& Headers) const
+{
+	XrCompositionLayerCylinderKHR Cylinder = { XR_TYPE_COMPOSITION_LAYER_CYLINDER_KHR, /*next*/ nullptr };
+	Cylinder.layerFlags = bNoAlpha ? 0 : XR_COMPOSITION_LAYER_UNPREMULTIPLIED_ALPHA_BIT |
+		XR_COMPOSITION_LAYER_BLEND_TEXTURE_SOURCE_ALPHA_BIT;
+	Cylinder.space = Space;
+	Cylinder.subImage.imageArrayIndex = 0;
+	Cylinder.pose = ToXrPose(Desc.Transform * PositionTransform, WorldToMeters);
+
+	const FCylinderLayer& CylinderProps = Desc.GetShape<FCylinderLayer>();
+	Cylinder.radius = FMath::Abs(CylinderProps.Radius / WorldToMeters);
+	Cylinder.centralAngle = FMath::Min((float)(2.0f * PI), FMath::Abs(CylinderProps.OverlayArc / CylinderProps.Radius));
+	Cylinder.aspectRatio = FMath::Abs(CylinderProps.OverlayArc / CylinderProps.Height);
+
+	FXrCompositionLayerUnion LayerUnion;
+	LayerUnion.Cylinder = Cylinder;
+
+	if (RightEye.Swapchain.IsValid())
+	{
+		LayerUnion.Cylinder.eyeVisibility = bIsStereo ? XR_EYE_VISIBILITY_RIGHT : XR_EYE_VISIBILITY_BOTH;
+		LayerUnion.Cylinder.subImage.imageRect = ToXrRect(GetRightViewportSize());
+		LayerUnion.Cylinder.subImage.swapchain = static_cast<FOpenXRSwapchain*>(RightEye.Swapchain.Get())->GetHandle();
+		Headers.Add(LayerUnion);
+	}
+	if (LeftEye.Swapchain.IsValid())
+	{
+		LayerUnion.Cylinder.eyeVisibility = XR_EYE_VISIBILITY_LEFT;
+		LayerUnion.Cylinder.subImage.imageRect = ToXrRect(GetLeftViewportSize());
+		LayerUnion.Cylinder.subImage.swapchain = static_cast<FOpenXRSwapchain*>(LeftEye.Swapchain.Get())->GetHandle();
+		Headers.Add(LayerUnion);
+	}
+}
+
+void FOpenXRLayer::CreateOpenXRQuadLayer(bool bIsStereo, bool bNoAlpha, FTransform PositionTransform, float WorldToMeters, XrSpace Space, TArray<FXrCompositionLayerUnion>& Headers) const
+{
+	XrCompositionLayerQuad Quad = { XR_TYPE_COMPOSITION_LAYER_QUAD, nullptr };
+	Quad.layerFlags = bNoAlpha ? 0 : XR_COMPOSITION_LAYER_UNPREMULTIPLIED_ALPHA_BIT |
+		XR_COMPOSITION_LAYER_BLEND_TEXTURE_SOURCE_ALPHA_BIT;
+	Quad.space = Space;
+	Quad.subImage.imageArrayIndex = 0;
+	Quad.pose = ToXrPose(Desc.Transform * PositionTransform, WorldToMeters);
+
+	// The layer pose doesn't take the transform scale into consideration, so we need to manually apply it the quad size.
+	const FVector2D LayerComponentScaler(Desc.Transform.GetScale3D().Y, Desc.Transform.GetScale3D().Z);
+
+	FXrCompositionLayerUnion LayerUnion;
+	LayerUnion.Quad = Quad;
+	
+	// We need to copy each layer into an OpenXR swapchain so they can be displayed by the compositor
+	if (RightEye.Swapchain.IsValid())
+	{
+		LayerUnion.Quad.eyeVisibility = bIsStereo ? XR_EYE_VISIBILITY_RIGHT : XR_EYE_VISIBILITY_BOTH;
+		LayerUnion.Quad.subImage.imageRect = ToXrRect(GetRightViewportSize());
+		LayerUnion.Quad.subImage.swapchain = static_cast<FOpenXRSwapchain*>(RightEye.Swapchain.Get())->GetHandle();
+		LayerUnion.Quad.size = ToXrExtent2D(GetRightQuadSize() * LayerComponentScaler, WorldToMeters);
+		Headers.Add(LayerUnion);
+	}
+	if (LeftEye.Swapchain.IsValid())
+	{
+		LayerUnion.Quad.eyeVisibility = XR_EYE_VISIBILITY_LEFT;
+		LayerUnion.Quad.subImage.imageRect = ToXrRect(GetLeftViewportSize());
+		LayerUnion.Quad.subImage.swapchain = static_cast<FOpenXRSwapchain*>(LeftEye.Swapchain.Get())->GetHandle();
+		LayerUnion.Quad.size = ToXrExtent2D(GetLeftQuadSize() * LayerComponentScaler, WorldToMeters);
+		Headers.Add(LayerUnion);
+	}
+}
+
+void FOpenXRLayer::CreateOpenXREquirectLayer(bool bIsStereo, bool bNoAlpha, FTransform PositionTransform, float WorldToMeters, XrSpace Space, TArray<FXrCompositionLayerUnion>& Headers) const
+{
+	XrCompositionLayerEquirectKHR Equirect = { XR_TYPE_COMPOSITION_LAYER_EQUIRECT_KHR, nullptr };
+	Equirect.layerFlags = bNoAlpha ? 0 : XR_COMPOSITION_LAYER_UNPREMULTIPLIED_ALPHA_BIT |
+		XR_COMPOSITION_LAYER_BLEND_TEXTURE_SOURCE_ALPHA_BIT;
+	Equirect.space = Space;
+	Equirect.subImage.imageArrayIndex = 0;
+	Equirect.pose = ToXrPose(Desc.Transform * PositionTransform, WorldToMeters);
+	// The radius is momentarily fixed to 0 as the editor does not support its insertion yet.
+	// An equirect layer with a radius of 0 is an infinite sphere.
+	Equirect.radius = 0.0f; 
+
+	const FEquirectLayer& EquirectProps = Desc.GetShape<FEquirectLayer>();
+
+	FXrCompositionLayerUnion LayerUnion;
+	LayerUnion.Equirect = Equirect;
+
+	// We need to copy each layer into an OpenXR swapchain so they can be displayed by the compositor
+	if (RightEye.Swapchain.IsValid())
+	{
+		LayerUnion.Equirect.eyeVisibility = bIsStereo ? XR_EYE_VISIBILITY_RIGHT : XR_EYE_VISIBILITY_BOTH;
+		LayerUnion.Equirect.subImage.imageRect = ToXrRect(GetRightViewportSize());
+		LayerUnion.Equirect.subImage.swapchain = static_cast<FOpenXRSwapchain*>(RightEye.Swapchain.Get())->GetHandle();
+		LayerUnion.Equirect.scale = ToXrVector2f(EquirectProps.RightScale);
+		LayerUnion.Equirect.bias = ToXrVector2f(EquirectProps.RightBias);
+		Headers.Add(LayerUnion);
+	}
+	if (LeftEye.Swapchain.IsValid())
+	{
+		LayerUnion.Equirect.eyeVisibility = XR_EYE_VISIBILITY_LEFT;
+		LayerUnion.Equirect.subImage.imageRect = ToXrRect(GetLeftViewportSize());
+		LayerUnion.Equirect.subImage.swapchain = static_cast<FOpenXRSwapchain*>(LeftEye.Swapchain.Get())->GetHandle();
+		LayerUnion.Equirect.scale = ToXrVector2f(EquirectProps.LeftScale);
+		LayerUnion.Equirect.bias = ToXrVector2f(EquirectProps.LeftBias);
+		Headers.Add(LayerUnion);
+	}
+}
 
 // TStereoLayerManager helper functions
 
diff --git a/Engine/Plugins/Runtime/OpenXR/Source/OpenXRHMD/Private/OpenXRHMD_Layer.h b/Engine/Plugins/Runtime/OpenXR/Source/OpenXRHMD/Private/OpenXRHMD_Layer.h
--- a/Engine/Plugins/Runtime/OpenXR/Source/OpenXRHMD/Private/OpenXRHMD_Layer.h
+++ b/Engine/Plugins/Runtime/OpenXR/Source/OpenXRHMD/Private/OpenXRHMD_Layer.h
@@ -4,8 +4,11 @@
 
 #include "CoreMinimal.h"
 #include "IStereoLayers.h"
+#include "OpenXRCore.h"
 #include "XRSwapChain.h"
 
+union FXrCompositionLayerUnion;
+
 struct FOpenXRLayer
 {
 	struct FPerEyeTextureData
@@ -45,6 +48,14 @@ struct FOpenXRLayer
 
 	FIntRect GetLeftViewportSize() const;
 	FVector2D GetLeftQuadSize() const;
+
+	TArray<FXrCompositionLayerUnion> CreateOpenXRLayer(FTransform InvTrackingToWorld, float WorldToMeters, XrSpace Space) const;
+	
+private:
+	void CreateOpenXRQuadLayer(bool bIsStereo, bool bNoAlpha, FTransform PositionTransform, float WorldToMeters, XrSpace Space, TArray<FXrCompositionLayerUnion>& Headers) const;
+	void CreateOpenXRCylinderLayer(bool bIsStereo, bool bNoAlpha, FTransform PositionTransform, float WorldToMeters, XrSpace Space, TArray<FXrCompositionLayerUnion>& Headers) const;
+	void CreateOpenXREquirectLayer(bool bIsStereo, bool bNoAlpha, FTransform PositionTransform, float WorldToMeters, XrSpace Space, TArray<FXrCompositionLayerUnion>& Headers) const;
+
 };
 
 bool GetLayerDescMember(const FOpenXRLayer& Layer, IStereoLayers::FLayerDesc& OutLayerDesc);
diff --git a/Engine/Plugins/Runtime/OpenXR/Source/OpenXRHMD/Public/OpenXRCore.h b/Engine/Plugins/Runtime/OpenXR/Source/OpenXRHMD/Public/OpenXRCore.h
--- a/Engine/Plugins/Runtime/OpenXR/Source/OpenXRHMD/Public/OpenXRCore.h
+++ b/Engine/Plugins/Runtime/OpenXR/Source/OpenXRHMD/Public/OpenXRCore.h
@@ -111,6 +111,11 @@ FORCEINLINE FVector2D ToFVector2D(XrVector2f Vector, float Scale = 1.0f)
 	return FVector2D(Vector.x * Scale, Vector.y * Scale);
 }
 
+FORCEINLINE XrVector2f ToXrVector2f(FVector2D Vector, float Scale = 1.0f)
+{
+	return XrVector2f{ (float)Vector.X / Scale, (float)Vector.Y / Scale };
+}
+
 FORCEINLINE FVector2D ToFVector2D(XrExtent2Df Extent, float Scale = 1.0f)
 {
 	return FVector2D(Extent.width * Scale, Extent.height * Scale);
@@ -131,6 +136,14 @@ FORCEINLINE uint32 ToXrPriority(int32 Priority)
 	return (uint32)Priority ^ (1 << 31);
 }
 
+union FXrCompositionLayerUnion
+{
+	XrCompositionLayerBaseHeader Header;
+	XrCompositionLayerQuad Quad;
+	XrCompositionLayerCylinderKHR Cylinder;
+	XrCompositionLayerEquirectKHR Equirect;
+};
+
 /**
  * XrPath wrapper with convenience functions
  */
