Avoid xr.OpenXRForceStereoLayerEmulation cvar changing at runtime.

diff --git a/Engine/Plugins/Runtime/OpenXR/Source/OpenXRHMD/Private/OpenXRHMD.cpp b/Engine/Plugins/Runtime/OpenXR/Source/OpenXRHMD/Private/OpenXRHMD.cpp
--- a/Engine/Plugins/Runtime/OpenXR/Source/OpenXRHMD/Private/OpenXRHMD.cpp
+++ b/Engine/Plugins/Runtime/OpenXR/Source/OpenXRHMD/Private/OpenXRHMD.cpp
@@ -66,7 +66,9 @@ static TAutoConsoleVariable<int32> CVarOpenXREnvironmentBlendMode(
 static TAutoConsoleVariable<bool> CVarOpenXRForceStereoLayerEmulation(
 	TEXT("xr.OpenXRForceStereoLayerEmulation"),
 	false,
-	TEXT("Force the emulation of stereo layers instead of using native ones (if supported).\n"),
+	TEXT("Force the emulation of stereo layers instead of using native ones (if supported).\n")
+	TEXT("The value of this cvar cannot be changed at runtime as it's cached during OnBeginPlay().\n")
+	TEXT("Any changes made at runtime will be picked up at the next VR Preview or app startup.\n"),
 	ECVF_Default);
 
 static TAutoConsoleVariable<bool> CVarOpenXRDoNotCopyEmulatedLayersToSpectatorScreen(
@@ -2157,6 +2159,7 @@ bool FOpenXRHMD::StopSession()
 
 void FOpenXRHMD::OnBeginPlay(FWorldContext& InWorldContext)
 {
+	bOpenXRForceStereoLayersEmulationCVarCachedValue = CVarOpenXRForceStereoLayerEmulation.GetValueOnGameThread();
 }
 
 void FOpenXRHMD::OnEndPlay(FWorldContext& InWorldContext)
@@ -2370,7 +2373,7 @@ void CreateNativeLayerSwapchain(FOpenXRLayer& Layer, TRefCountPtr<FOpenXRRenderB
 
 bool FOpenXRHMD::IsEmulatingStereoLayers()
 {
-	return !bLayerSupportOpenXRCompliant || CVarOpenXRForceStereoLayerEmulation.GetValueOnAnyThread();
+	return !bLayerSupportOpenXRCompliant || bOpenXRForceStereoLayersEmulationCVarCachedValue;
 }
 
 void FOpenXRHMD::SetupFrameLayers_RenderThread(FRHICommandListImmediate& RHICmdList)
diff --git a/Engine/Plugins/Runtime/OpenXR/Source/OpenXRHMD/Private/OpenXRHMD.h b/Engine/Plugins/Runtime/OpenXR/Source/OpenXRHMD/Private/OpenXRHMD.h
--- a/Engine/Plugins/Runtime/OpenXR/Source/OpenXRHMD/Private/OpenXRHMD.h
+++ b/Engine/Plugins/Runtime/OpenXR/Source/OpenXRHMD/Private/OpenXRHMD.h
@@ -467,6 +467,7 @@ private:
 	FVector					BasePosition;
 
 	bool					bLayerSupportOpenXRCompliant;
+	bool					bOpenXRForceStereoLayersEmulationCVarCachedValue;
 	TArray<IStereoLayers::FLayerDesc> BackgroundCompositedEmulatedLayers;
 	TArray<IStereoLayers::FLayerDesc> EmulatedFaceLockedLayers;
 	TArray<FOpenXRLayer>			  NativeLayers;
