Add Palm Pose support to OpenXR

diff --git a/Engine/Plugins/Runtime/OpenXR/Source/OpenXRHMD/Private/OpenXRHMD.cpp b/Engine/Plugins/Runtime/OpenXR/Source/OpenXRHMD/Private/OpenXRHMD.cpp
--- a/Engine/Plugins/Runtime/OpenXR/Source/OpenXRHMD/Private/OpenXRHMD.cpp
+++ b/Engine/Plugins/Runtime/OpenXR/Source/OpenXRHMD/Private/OpenXRHMD.cpp
@@ -234,10 +234,9 @@ void FOpenXRHMD::GetMotionControllerData(UObject* WorldContext, const EControlle
 			}
 		}
 
+		const float WorldToMeters = GetWorldToMetersScale();
 		if (MotionController)
 		{
-			const float WorldToMeters = GetWorldToMetersScale();
-
 			bool bSuccess = false;
 			FVector Position = FVector::ZeroVector;
 			FRotator Rotation = FRotator::ZeroRotator;
@@ -260,6 +259,15 @@ void FOpenXRHMD::GetMotionControllerData(UObject* WorldContext, const EControlle
 			}
 			MotionControllerData.bValid |= bSuccess;
 
+			FName PalmSource = Hand == EControllerHand::Left ? FName("LeftPalm") : FName("RightPalm");
+			bSuccess = MotionController->GetControllerOrientationAndPosition(0, PalmSource, Rotation, Position, WorldToMeters);
+			if (bSuccess)
+			{
+				MotionControllerData.GripPosition = trackingToWorld.TransformPosition(Position);
+				MotionControllerData.GripRotation = trackingToWorld.TransformRotation(FQuat(Rotation));
+			}
+			MotionControllerData.bValid |= bSuccess;
+
 			MotionControllerData.TrackingStatus = MotionController->GetControllerTrackingStatus(0, GripSource);
 		}
 
@@ -270,7 +278,7 @@ void FOpenXRHMD::GetMotionControllerData(UObject* WorldContext, const EControlle
 
 			MotionControllerData.bValid = HandTracker->GetAllKeypointStates(Hand, MotionControllerData.HandKeyPositions, MotionControllerData.HandKeyRotations, MotionControllerData.HandKeyRadii);
 			check(!MotionControllerData.bValid || (MotionControllerData.HandKeyPositions.Num() == EHandKeypointCount && MotionControllerData.HandKeyRotations.Num() == EHandKeypointCount && MotionControllerData.HandKeyRadii.Num() == EHandKeypointCount));
-		}	
+		}
 	}
 
 	//TODO: this is reportedly a wmr specific convenience function for rapid prototyping.  Not sure it is useful for openxr.
diff --git a/Engine/Plugins/Runtime/OpenXR/Source/OpenXRHMD/Private/OpenXRHMDModule.cpp b/Engine/Plugins/Runtime/OpenXR/Source/OpenXRHMD/Private/OpenXRHMDModule.cpp
--- a/Engine/Plugins/Runtime/OpenXR/Source/OpenXRHMD/Private/OpenXRHMDModule.cpp
+++ b/Engine/Plugins/Runtime/OpenXR/Source/OpenXRHMD/Private/OpenXRHMDModule.cpp
@@ -498,6 +498,7 @@ bool FOpenXRHMDModule::GetOptionalExtensions(TArray<const ANSICHAR*>& OutExtensi
 	OutExtensions.Add(XR_KHR_VISIBILITY_MASK_EXTENSION_NAME);
 	OutExtensions.Add(XR_KHR_BINDING_MODIFICATION_EXTENSION_NAME);
 	OutExtensions.Add(XR_EPIC_VIEW_CONFIGURATION_FOV_EXTENSION_NAME);
+	OutExtensions.Add(XR_EXT_PALM_POSE_EXTENSION_NAME);
 
 	// Draft extension not yet provided in headers
 	OutExtensions.Add("XR_EXT_dpad_binding");
diff --git a/Engine/Plugins/Runtime/OpenXR/Source/OpenXRInput/Private/OpenXRInput.cpp b/Engine/Plugins/Runtime/OpenXR/Source/OpenXRInput/Private/OpenXRInput.cpp
--- a/Engine/Plugins/Runtime/OpenXR/Source/OpenXRInput/Private/OpenXRInput.cpp
+++ b/Engine/Plugins/Runtime/OpenXR/Source/OpenXRInput/Private/OpenXRInput.cpp
@@ -37,6 +37,8 @@ namespace OpenXRSourceNames
 	static const FName RightGrip("RightGrip");
 	static const FName LeftAim("LeftAim");
 	static const FName RightAim("RightAim");
+	static const FName LeftPalm("LeftPalm");
+	static const FName RightPalm("RightPalm");
 }
 
 FORCEINLINE XrActionType ToActionType(EInputActionValueType ValueType)
@@ -195,6 +197,12 @@ FOpenXRInputPlugin::FOpenXRController::FOpenXRController(XrActionSet InActionSet
 	Info.actionType = XR_ACTION_TYPE_POSE_INPUT;
 	XR_ENSURE(xrCreateAction(ActionSet, &Info, &AimAction));
 
+	FCStringAnsi::Strcpy(Info.localizedActionName, XR_MAX_ACTION_NAME_SIZE, InName);
+	FCStringAnsi::Strcat(Info.localizedActionName, XR_MAX_ACTION_NAME_SIZE, " Palm Pose");
+	FilterActionName(Info.localizedActionName, Info.actionName);
+	Info.actionType = XR_ACTION_TYPE_POSE_INPUT;
+	XR_ENSURE(xrCreateAction(ActionSet, &Info, &PalmAction));
+
 	FCStringAnsi::Strcpy(Info.localizedActionName, XR_MAX_ACTION_NAME_SIZE, InName);
 	FCStringAnsi::Strcat(Info.localizedActionName, XR_MAX_ACTION_NAME_SIZE, " Vibration");
 	FilterActionName(Info.localizedActionName, Info.actionName);
@@ -208,6 +216,7 @@ void FOpenXRInputPlugin::FOpenXRController::AddTrackedDevices(FOpenXRHMD* HMD)
 	{
 		GripDeviceId = HMD->AddTrackedDevice(GripAction, UserPath);
 		AimDeviceId = HMD->AddTrackedDevice(AimAction, UserPath);
+		PalmDeviceId = HMD->AddTrackedDevice(PalmAction, UserPath);
 	}
 }
 
@@ -231,6 +240,7 @@ FOpenXRInputPlugin::FOpenXRInput::FOpenXRInput(FOpenXRHMD* HMD)
 	, MappableInputConfig(nullptr)
 	, bActionsAttached(false)
 	, bDirectionalBindingSupported(false)
+	, bPalmPoseSupported(false)
 	, MessageHandler(new FGenericApplicationMessageHandler())
 {
 	IModularFeatures::Get().RegisterModularFeature(GetModularFeatureName(), this);
@@ -240,6 +250,7 @@ FOpenXRInputPlugin::FOpenXRInput::FOpenXRInput(FOpenXRHMD* HMD)
 	{
 		Instance = OpenXRHMD->GetInstance();
 		bDirectionalBindingSupported = OpenXRHMD->IsExtensionEnabled("XR_EXT_dpad_binding");
+		bPalmPoseSupported = OpenXRHMD->IsExtensionEnabled("XR_EXT_palm_pose");
 
 		// Note: AnyHand needs special handling because it tries left then falls back to right in each call.
 		MotionSourceToControllerHandMap.Add(OpenXRSourceNames::Left, EControllerHand::Left);
@@ -248,6 +259,8 @@ FOpenXRInputPlugin::FOpenXRInput::FOpenXRInput(FOpenXRHMD* HMD)
 		MotionSourceToControllerHandMap.Add(OpenXRSourceNames::RightGrip, EControllerHand::Right);
 		MotionSourceToControllerHandMap.Add(OpenXRSourceNames::LeftAim, EControllerHand::Left);
 		MotionSourceToControllerHandMap.Add(OpenXRSourceNames::RightAim, EControllerHand::Right);
+		MotionSourceToControllerHandMap.Add(OpenXRSourceNames::LeftPalm, EControllerHand::Left);
+		MotionSourceToControllerHandMap.Add(OpenXRSourceNames::RightPalm, EControllerHand::Right);
 
 		// Map the legacy hand enum values that openxr supports
 		MotionSourceToControllerHandMap.Add(TEXT("EControllerHand::Left"), EControllerHand::Left);
@@ -274,6 +287,10 @@ XrAction FOpenXRInputPlugin::FOpenXRInput::GetActionForMotionSource(FName Motion
 	{
 		return Controller.AimAction;
 	}
+	else if (MotionSource == OpenXRSourceNames::LeftPalm || MotionSource == OpenXRSourceNames::RightPalm)
+	{
+		return Controller.PalmAction;
+	}
 	else
 	{
 		return Controller.GripAction;
@@ -287,6 +304,10 @@ int32 FOpenXRInputPlugin::FOpenXRInput::GetDeviceIDForMotionSource(FName MotionS
 	{
 		return Controller.AimDeviceId;
 	}
+	else if (MotionSource == OpenXRSourceNames::LeftPalm || MotionSource == OpenXRSourceNames::RightPalm)
+	{
+		return Controller.PalmDeviceId;
+	}
 	else
 	{
 		return Controller.GripDeviceId;
@@ -395,6 +416,25 @@ bool FOpenXRInputPlugin::FOpenXRInput::BuildActions(XrSession Session)
 				Controllers[EControllerHand::Right].AimAction, FOpenXRPath("/user/hand/right/input/aim/pose")
 				});
 
+			if (bPalmPoseSupported)
+			{
+				Profile.Bindings.Add(XrActionSuggestedBinding{
+					Controllers[EControllerHand::Left].PalmAction, FOpenXRPath("/user/hand/left/input/palm_ext/pose")
+					});
+				Profile.Bindings.Add(XrActionSuggestedBinding{
+					Controllers[EControllerHand::Right].PalmAction, FOpenXRPath("/user/hand/right/input/palm_ext/pose")
+					});
+			}
+			else
+			{
+				Profile.Bindings.Add(XrActionSuggestedBinding{
+					Controllers[EControllerHand::Left].PalmAction, FOpenXRPath("/user/hand/left/input/grip/pose")
+					});
+				Profile.Bindings.Add(XrActionSuggestedBinding{
+					Controllers[EControllerHand::Right].PalmAction, FOpenXRPath("/user/hand/right/input/grip/pose")
+					});
+			}
+
 			if (Profile.HasHaptics)
 			{
 				Profile.Bindings.Add(XrActionSuggestedBinding{
@@ -1256,6 +1296,8 @@ void FOpenXRInputPlugin::FOpenXRInput::EnumerateSources(TArray<FMotionController
 	SourcesOut.Add(OpenXRSourceNames::RightGrip);
 	SourcesOut.Add(OpenXRSourceNames::LeftAim);
 	SourcesOut.Add(OpenXRSourceNames::RightAim);
+	SourcesOut.Add(OpenXRSourceNames::LeftPalm);
+	SourcesOut.Add(OpenXRSourceNames::RightPalm);
 }
 
 // TODO: Refactor API to change the Hand type to EControllerHand
diff --git a/Engine/Plugins/Runtime/OpenXR/Source/OpenXRInput/Private/OpenXRInput.h b/Engine/Plugins/Runtime/OpenXR/Source/OpenXRInput/Private/OpenXRInput.h
--- a/Engine/Plugins/Runtime/OpenXR/Source/OpenXRInput/Private/OpenXRInput.h
+++ b/Engine/Plugins/Runtime/OpenXR/Source/OpenXRInput/Private/OpenXRInput.h
@@ -79,9 +79,11 @@ public:
 		XrPath			UserPath;
 		XrAction		GripAction;
 		XrAction		AimAction;
+		XrAction		PalmAction;
 		XrAction		VibrationAction;
 		int32			GripDeviceId;
 		int32			AimDeviceId;
+		int32			PalmDeviceId;
 
 		bool			bHapticActive;
 
@@ -156,6 +158,7 @@ public:
 
 		bool bActionsAttached;
 		bool bDirectionalBindingSupported;
+		bool bPalmPoseSupported;
 
 		/**
 		* Buffer for current delta time to get an accurate approximation of how long to play haptics for
diff --git a/Engine/Source/Runtime/HeadMountedDisplay/Public/HeadMountedDisplayTypes.h b/Engine/Source/Runtime/HeadMountedDisplay/Public/HeadMountedDisplayTypes.h
--- a/Engine/Source/Runtime/HeadMountedDisplay/Public/HeadMountedDisplayTypes.h
+++ b/Engine/Source/Runtime/HeadMountedDisplay/Public/HeadMountedDisplayTypes.h
@@ -352,18 +352,28 @@ struct HEADMOUNTEDDISPLAY_API FXRMotionControllerData
 
 	UPROPERTY(BlueprintReadOnly, Category = "XR")
 	ETrackingStatus TrackingStatus = ETrackingStatus::NotTracked;
-	
+
+	// Vector representing an object being held in the player's hand
 	UPROPERTY(BlueprintReadOnly, Category = "XR")
 	FVector GripPosition = FVector(0.0f);
+	// Quaternion representing an object being held in the player's hand
 	UPROPERTY(BlueprintReadOnly, Category = "XR")
 	FQuat GripRotation = FQuat(EForceInit::ForceInitToZero);
 
-	//for hand controllers, provides a more steady vector based on the elbow
+	// For handheld controllers, gives a vector for pointing at objects
 	UPROPERTY(BlueprintReadOnly, Category = "XR")
-	FVector AimPosition = FVector(0.0f);;
+	FVector AimPosition = FVector(0.0f);
+	// For handheld controllers, gives a quaternion for pointing at objects
 	UPROPERTY(BlueprintReadOnly, Category = "XR")
 	FQuat AimRotation = FQuat(EForceInit::ForceInitToZero);
 
+	// For handheld controllers, gives a vector for representing the hand
+	UPROPERTY(BlueprintReadOnly, Category = "XR")
+	FVector PalmPosition = FVector(0.0f);
+	// For handheld controllers, gives a quaternion for representing the hand
+	UPROPERTY(BlueprintReadOnly, Category = "XR")
+	FQuat PalmRotation = FQuat(EForceInit::ForceInitToZero);
+
 	// The indices of this array are the values of EHandKeypoint (Palm, Wrist, ThumbMetacarpal, etc).
 	UPROPERTY(BlueprintReadOnly, Category = "XR")
 	TArray<FVector> HandKeyPositions;
