OpenXR: disable depth swapchain with ScreenPercentage != 100 + layer state flags

If we allocate a depth swapchain, it cannot change size with screen percentage (instead uses PixelDensity).
If the SceneTextures end up changing size, there will be a mismatch between the scene textures and the depth buffer sourced from the OpenXR depth swapchain.
There are a couple solutions we can pursue, but for now, it's simplest to disable depth submission because it's optional anyway.

Because of the engine ordering of when a view family starts rendering and the scene textures are allocated, we can't disable the depth swapchain outright.
Instead, we just won't submit it to the compositor. This means we need to differentiate between swapchain management and swapchain submission.
This change introduces layer state flags, which consolidates some existing and new logic to signal per-frame layer management.
We recompute this state every frame, a net positive compared to the inconsistent retained + dynamic state.

diff --git a/Engine/Plugins/Runtime/OpenXR/Source/OpenXRHMD/Private/OpenXRHMD.cpp b/Engine/Plugins/Runtime/OpenXR/Source/OpenXRHMD/Private/OpenXRHMD.cpp
--- a/Engine/Plugins/Runtime/OpenXR/Source/OpenXRHMD/Private/OpenXRHMD.cpp
+++ b/Engine/Plugins/Runtime/OpenXR/Source/OpenXRHMD/Private/OpenXRHMD.cpp
@@ -960,13 +960,13 @@ void FOpenXRHMD::SetFinalViewRect(FRHICommandListImmediate& RHICmdList, const in
 	};
 
 	XrSwapchainSubImage& DepthImage = PipelinedLayerStateRendering.DepthImages[ViewIndex];
-	if (bDepthExtensionSupported)
+	if (EnumHasAnyFlags(PipelinedLayerStateRendering.LayerStateFlags, EOpenXRLayerStateFlags::SubmitDepthLayer))
 	{
 		DepthImage.imageArrayIndex = ColorImage.imageArrayIndex;
 		DepthImage.imageRect = ColorImage.imageRect;
 	}
 	XrSwapchainSubImage& EmulationImage = PipelinedLayerStateRendering.EmulatedLayerState.EmulationImages[ViewIndex];
-	if(PipelinedLayerStateRendering.EmulatedLayerState.bIsFaceLockedLayerEmulationActive)
+	if (EnumHasAnyFlags(PipelinedLayerStateRendering.LayerStateFlags, EOpenXRLayerStateFlags::SubmitEmulatedFaceLockedLayer))
 	{
 		EmulationImage.imageArrayIndex = ColorImage.imageArrayIndex;
 		EmulationImage.imageRect = ColorImage.imageRect;
@@ -1156,11 +1156,11 @@ void FOpenXRHMD::PostRenderViewFamily_RenderThread(FRDGBuilder& GraphBuilder, FS
 		XrSwapchainSubImage& EmulationImage = PipelinedLayerStateRendering.EmulatedLayerState.EmulationImages[ViewIndex];
 
 		ColorImage.swapchain = PipelinedLayerStateRendering.ColorSwapchain.IsValid() ? static_cast<FOpenXRSwapchain*>(PipelinedLayerStateRendering.ColorSwapchain.Get())->GetHandle() : XR_NULL_HANDLE;
-		if (bDepthExtensionSupported)
+		if (EnumHasAnyFlags(PipelinedLayerStateRendering.LayerStateFlags, EOpenXRLayerStateFlags::SubmitDepthLayer))
 		{
 			DepthImage.swapchain = PipelinedLayerStateRendering.DepthSwapchain.IsValid() ? static_cast<FOpenXRSwapchain*>(PipelinedLayerStateRendering.DepthSwapchain.Get())->GetHandle() : XR_NULL_HANDLE;
 		}
-		if(PipelinedLayerStateRendering.EmulatedLayerState.bIsFaceLockedLayerEmulationActive)
+		if (EnumHasAnyFlags(PipelinedLayerStateRendering.LayerStateFlags, EOpenXRLayerStateFlags::SubmitEmulatedFaceLockedLayer))
 		{
 			EmulationImage.swapchain = PipelinedLayerStateRendering.EmulatedLayerState.EmulationSwapchain.IsValid() ? static_cast<FOpenXRSwapchain*>(PipelinedLayerStateRendering.EmulatedLayerState.EmulationSwapchain.Get())->GetHandle() : XR_NULL_HANDLE;
 		}
@@ -1177,7 +1177,7 @@ void FOpenXRHMD::PostRenderViewFamily_RenderThread(FRDGBuilder& GraphBuilder, FS
 		Projection.next = nullptr;
 		Projection.subImage = ColorImage;
 
-		if (bDepthExtensionSupported && PipelinedLayerStateRendering.DepthSwapchain.IsValid())
+		if (EnumHasAnyFlags(PipelinedLayerStateRendering.LayerStateFlags, EOpenXRLayerStateFlags::SubmitDepthLayer))
 		{
 			XrCompositionLayerDepthInfoKHR& DepthLayer = PipelinedLayerStateRendering.DepthLayers[ViewIndex];
 
@@ -1196,7 +1196,7 @@ void FOpenXRHMD::PostRenderViewFamily_RenderThread(FRDGBuilder& GraphBuilder, FS
 
 			Projection.next = &DepthLayer;
 		}
-		if (PipelinedLayerStateRendering.EmulatedLayerState.bIsFaceLockedLayerEmulationActive)
+		if (EnumHasAnyFlags(PipelinedLayerStateRendering.LayerStateFlags, EOpenXRLayerStateFlags::SubmitEmulatedFaceLockedLayer))
 		{
 			XrCompositionLayerProjectionView& CompositedProjection = PipelinedLayerStateRendering.EmulatedLayerState.CompositedProjectionLayers[ViewIndex];
 
@@ -2242,8 +2242,9 @@ bool FOpenXRHMD::AllocateRenderTargetTexture(uint32 Index, uint32 SizeX, uint32
 	LastRequestedColorSwapchainFormat = Format;
 	LastActualColorSwapchainFormat = ActualFormat;
 
-	if(IsEmulatingStereoLayers())
+	if(IsEmulatingStereoLayers() && (SystemProperties.graphicsProperties.maxLayerCount > 1))
 	{
+		// If we have at least two native layers, use non-background layer to render the composited image of all the emulated face locked layers.
 		FXRSwapChainPtr& EmulationSwapchain = PipelinedLayerStateRendering.EmulatedLayerState.EmulationSwapchain;
 		const FRHITexture2D* const EmulationSwapchainTexture = EmulationSwapchain == nullptr ? nullptr : EmulationSwapchain->GetTexture2DArray() ? EmulationSwapchain->GetTexture2DArray() : EmulationSwapchain->GetTexture2D();
 		if (EmulationSwapchain == nullptr || EmulationSwapchainTexture == nullptr || EmulationSwapchainTexture->GetSizeX() != SizeX || EmulationSwapchainTexture->GetSizeY() != SizeY)
@@ -2331,6 +2332,12 @@ bool FOpenXRHMD::AllocateDepthTexture(uint32 Index, uint32 SizeX, uint32 SizeY,
 	const FRHITexture2D* const DepthSwapchainTexture = DepthSwapchain->GetTexture2DArray() ? DepthSwapchain->GetTexture2DArray() : DepthSwapchain->GetTexture2D();
 	const FRHITextureDesc& DepthSwapchainDesc = DepthSwapchainTexture->GetDesc();
 
+	if (SizeX != DepthSwapchainDesc.Extent.X || SizeY != DepthSwapchainDesc.Extent.Y)
+	{
+		// We don't yet support different sized SceneTexture depth + OpenXR layer depth
+		return false;
+	}
+
 	// Sample count, mip count and size should be known at AllocateRenderTargetTexture time
 	// Format _could_ change, but we should know it (and can check for it in AllocateDepthTextureInternal)
 	// Flags might also change. We expect TexCreate_DepthStencilTargetable | TexCreate_ShaderResource | TexCreate_InputAttachmentRead from SceneTextures
@@ -2343,6 +2350,8 @@ bool FOpenXRHMD::AllocateDepthTexture(uint32 Index, uint32 SizeX, uint32 SizeY,
 
 	OutTargetableTexture = OutShaderResourceTexture = (FTexture2DRHIRef&)PipelinedLayerStateRendering.DepthSwapchain->GetTextureRef();
 
+	PipelinedLayerStateRendering.LayerStateFlags |= EOpenXRLayerStateFlags::SubmitDepthLayer;
+
 	return true;
 }
 
@@ -2415,7 +2424,7 @@ void FOpenXRHMD::SetupFrameLayers_RenderThread(FRHICommandListImmediate& RHICmdL
 				{
 					// If we have at least one native layer, use it to render the 
 					// composited image of all the emulated face locked layers.
-					if (SystemProperties.graphicsProperties.maxLayerCount > 0)
+					if (PipelinedLayerStateRendering.EmulatedLayerState.EmulationSwapchain.IsValid())
 					{
 						EmulatedFaceLockedLayers.Add(Layer.Desc);
 					}
@@ -2461,10 +2470,10 @@ void FOpenXRHMD::SetupFrameLayers_RenderThread(FRHICommandListImmediate& RHICmdL
 		BackgroundCompositedEmulatedLayers.Sort(LayerCompare);
 		EmulatedFaceLockedLayers.Sort(LayerCompare);
 		NativeLayers.Sort(LayerCompare);
-
-		PipelinedLayerStateRendering.EmulatedLayerState.bIsFaceLockedLayerEmulationActive = !EmulatedFaceLockedLayers.IsEmpty();
 	}
 
+	PipelinedLayerStateRendering.LayerStateFlags |= !EmulatedFaceLockedLayers.IsEmpty() ? EOpenXRLayerStateFlags::SubmitEmulatedFaceLockedLayer : EOpenXRLayerStateFlags::None;
+
 	const FTransform InvTrackingToWorld = GetTrackingToWorldTransform().Inverse();
 	const float WorldToMeters = GetWorldToMetersScale();
 
@@ -2561,7 +2570,7 @@ void FOpenXRHMD::DrawEmulatedLayers_RenderThread(FRDGBuilder& GraphBuilder, cons
 }
 void FOpenXRHMD::DrawEmulatedFaceLockedLayers_RenderThread(FRDGBuilder& GraphBuilder, const FSceneView& InView)
 {
-	if (!PipelinedLayerStateRendering.EmulatedLayerState.bIsFaceLockedLayerEmulationActive || !PipelinedLayerStateRendering.EmulatedLayerState.EmulationSwapchain)
+	if (!EnumHasAnyFlags(PipelinedLayerStateRendering.LayerStateFlags, EOpenXRLayerStateFlags::SubmitEmulatedFaceLockedLayer))
 	{
 		return;
 	}
@@ -2598,7 +2607,7 @@ void FOpenXRHMD::DrawBackgroundCompositedEmulatedLayers_RenderThread(FRDGBuilder
 		RHICmdList.BeginRenderPass(RPInfo, TEXT("EmulatedStereoLayerRender"));
 		RHICmdList.SetViewport((float)RenderParams.Viewport.Min.X, (float)RenderParams.Viewport.Min.Y, 0.0f, (float)RenderParams.Viewport.Max.X, (float)RenderParams.Viewport.Max.Y, 1.0f);
 
-		if (bSplashIsShown || !PipelinedLayerStateRendering.bBackgroundLayerVisible)
+		if (bSplashIsShown || !EnumHasAnyFlags(PipelinedLayerStateRendering.LayerStateFlags, EOpenXRLayerStateFlags::BackgroundLayerVisible))
 		{
 			DrawClearQuad(RHICmdList, FLinearColor::Black);
 		}
@@ -2675,6 +2684,8 @@ void FOpenXRHMD::OnBeginRendering_RenderThread(FRHICommandListImmediate& RHICmdL
 		Module->OnBeginRendering_RenderThread(Session);
 	}
 
+	SetupFrameLayers_RenderThread(RHICmdList);
+
 	const float WorldToMeters = GetWorldToMetersScale();
 
 	for (int32 ViewIndex = 0; ViewIndex < ViewFamily.Views.Num(); ViewIndex++)
@@ -2696,7 +2707,7 @@ void FOpenXRHMD::OnBeginRendering_RenderThread(FRHICommandListImmediate& RHICmdL
 		Projection.pose = ToXrPose(BasePoseTransform, WorldToMeters);
 		Projection.fov = View.fov;
 
-		if(PipelinedLayerStateRendering.EmulatedLayerState.bIsFaceLockedLayerEmulationActive)
+		if (EnumHasAnyFlags(PipelinedLayerStateRendering.LayerStateFlags, EOpenXRLayerStateFlags::SubmitEmulatedFaceLockedLayer))
 		{
 			XrCompositionLayerProjectionView& CompositedProjection = PipelinedLayerStateRendering.EmulatedLayerState.CompositedProjectionLayers[ViewIndex];
 			CompositedProjection.pose = ToXrPose(EyePose, WorldToMeters);
@@ -2704,8 +2715,6 @@ void FOpenXRHMD::OnBeginRendering_RenderThread(FRHICommandListImmediate& RHICmdL
 		}
 	}
 
-	SetupFrameLayers_RenderThread(RHICmdList);
-
 #if !PLATFORM_HOLOLENS
 	if (bHiddenAreaMaskSupported && bNeedReBuildOcclusionMesh)
 	{
@@ -2736,10 +2745,8 @@ void FOpenXRHMD::OnBeginRendering_RenderThread(FRHICommandListImmediate& RHICmdL
 		// Always check for sanity before using it.
 		FXRSwapChainPtr EmulationSwapchain = PipelinedLayerStateRendering.EmulatedLayerState.EmulationSwapchain;
 
-		RHICmdList.EnqueueLambda([this, FrameState = PipelinedFrameStateRendering, bIsFaceLockedLayerEmulationActive = PipelinedLayerStateRendering.EmulatedLayerState.bIsFaceLockedLayerEmulationActive, ColorSwapchain, DepthSwapchain, EmulationSwapchain](FRHICommandListImmediate& InRHICmdList)
+		RHICmdList.EnqueueLambda([this, FrameState = PipelinedFrameStateRendering, ColorSwapchain, DepthSwapchain, EmulationSwapchain](FRHICommandListImmediate& InRHICmdList)
 		{
-			// This needs to be passed here or the first frame will not have the correct value and will not acquire the swapchain image even if the emulation is active.
-			PipelinedLayerStateRHI.EmulatedLayerState.bIsFaceLockedLayerEmulationActive = bIsFaceLockedLayerEmulationActive;
 			OnBeginRendering_RHIThread(FrameState, ColorSwapchain, DepthSwapchain, EmulationSwapchain);
 		});
 
@@ -2804,7 +2811,7 @@ void FOpenXRHMD::OnLateUpdateApplied_RenderThread(FRHICommandListImmediate& RHIC
 		// Update the field-of-view to match the final projection matrix
 		Projection.fov = View.fov;
 
-		if(PipelinedLayerStateRendering.EmulatedLayerState.bIsFaceLockedLayerEmulationActive)
+		if (EnumHasAnyFlags(PipelinedLayerStateRendering.LayerStateFlags, EOpenXRLayerStateFlags::SubmitEmulatedFaceLockedLayer))
 		{
 			XrCompositionLayerProjectionView& CompositedProjection = PipelinedLayerStateRendering.EmulatedLayerState.CompositedProjectionLayers[ViewIndex];
 			CompositedProjection.pose = ToXrPose(EyePose, GetWorldToMetersScale());
@@ -2834,9 +2841,12 @@ void FOpenXRHMD::OnBeginRendering_GameThread()
 
 			PipelinedFrameStateRendering = GameFrameState;
 
+			PipelinedLayerStateRendering.LayerStateFlags = EOpenXRLayerStateFlags::None;
+
 			// If we are emulating layers, we still need to submit background layer since we composite into it
-			PipelinedLayerStateRendering.bBackgroundLayerVisible = bBackgroundLayerVisible;
-			PipelinedLayerStateRendering.bSubmitBackgroundLayer = bBackgroundLayerVisible || IsEmulatingStereoLayers();
+			PipelinedLayerStateRendering.LayerStateFlags |= bBackgroundLayerVisible ? EOpenXRLayerStateFlags::BackgroundLayerVisible : EOpenXRLayerStateFlags::None;
+			PipelinedLayerStateRendering.LayerStateFlags |= (bBackgroundLayerVisible || IsEmulatingStereoLayers()) ?
+				EOpenXRLayerStateFlags::SubmitBackgroundLayer : EOpenXRLayerStateFlags::None;
 		});
 }
 
@@ -3146,12 +3156,12 @@ void FOpenXRHMD::OnBeginRendering_RHIThread(const FPipelinedFrameState& InFrameS
 		{
 			ColorSwapchain->IncrementSwapChainIndex_RHIThread();
 			ColorSwapchain->WaitCurrentImage_RHIThread(OPENXR_SWAPCHAIN_WAIT_TIMEOUT);
-			if (bDepthExtensionSupported && DepthSwapchain)
+			if (DepthSwapchain)
 			{
 				DepthSwapchain->IncrementSwapChainIndex_RHIThread();
 				DepthSwapchain->WaitCurrentImage_RHIThread(OPENXR_SWAPCHAIN_WAIT_TIMEOUT);
 			}
-			if(PipelinedLayerStateRHI.EmulatedLayerState.bIsFaceLockedLayerEmulationActive && EmulationSwapchain)
+			if (EmulationSwapchain)
 			{
 				EmulationSwapchain->IncrementSwapChainIndex_RHIThread();
 				EmulationSwapchain->WaitCurrentImage_RHIThread(OPENXR_SWAPCHAIN_WAIT_TIMEOUT);
@@ -3189,11 +3199,11 @@ void FOpenXRHMD::OnFinishRendering_RHIThread()
 	{
 		PipelinedLayerStateRHI.ColorSwapchain->ReleaseCurrentImage_RHIThread();
 
-		if (bDepthExtensionSupported && PipelinedLayerStateRHI.DepthSwapchain)
+		if (PipelinedLayerStateRHI.DepthSwapchain)
 		{
 			PipelinedLayerStateRHI.DepthSwapchain->ReleaseCurrentImage_RHIThread();
 		}
-		if(PipelinedLayerStateRHI.EmulatedLayerState.bIsFaceLockedLayerEmulationActive && PipelinedLayerStateRHI.EmulatedLayerState.EmulationSwapchain)
+		if (PipelinedLayerStateRHI.EmulatedLayerState.EmulationSwapchain)
 		{
 			PipelinedLayerStateRHI.EmulatedLayerState.EmulationSwapchain->ReleaseCurrentImage_RHIThread();
 		}
@@ -3204,7 +3214,7 @@ void FOpenXRHMD::OnFinishRendering_RHIThread()
 	{
 		TArray<const XrCompositionLayerBaseHeader*> Headers;
 		XrCompositionLayerProjection Layer = {};
-		if (PipelinedLayerStateRHI.bSubmitBackgroundLayer)
+		if (EnumHasAnyFlags(PipelinedLayerStateRHI.LayerStateFlags, EOpenXRLayerStateFlags::SubmitBackgroundLayer))
 		{
 			Layer.type = XR_TYPE_COMPOSITION_LAYER_PROJECTION;
 			Layer.next = nullptr;
@@ -3226,7 +3236,7 @@ void FOpenXRHMD::OnFinishRendering_RHIThread()
 		}
 		
 		XrCompositionLayerProjection CompositedLayer = {};
-		if (PipelinedLayerStateRHI.EmulatedLayerState.bIsFaceLockedLayerEmulationActive && PipelinedLayerStateRHI.EmulatedLayerState.EmulationSwapchain)
+		if (EnumHasAnyFlags(PipelinedLayerStateRHI.LayerStateFlags, EOpenXRLayerStateFlags::SubmitEmulatedFaceLockedLayer))
 		{
 			CompositedLayer.type = XR_TYPE_COMPOSITION_LAYER_PROJECTION;
 			CompositedLayer.next = nullptr;
@@ -3264,7 +3274,7 @@ void FOpenXRHMD::OnFinishRendering_RHIThread()
 				if (PipelinedFrameStateRHI.PluginViewInfos[i].Plugin == Module && PipelinedLayerStateRHI.ColorImages.IsValidIndex(i))
 				{
 					ColorImages.Add(PipelinedLayerStateRHI.ColorImages[i]);
-					if (bDepthExtensionSupported)
+					if (EnumHasAnyFlags(PipelinedLayerStateRHI.LayerStateFlags, EOpenXRLayerStateFlags::SubmitDepthLayer))
 					{
 						DepthImages.Add(PipelinedLayerStateRHI.DepthImages[i]);
 					}
@@ -3724,7 +3734,7 @@ void FOpenXRHMD::RenderTexture_RenderThread(class FRHICommandListImmediate& RHIC
 {
 	if (SpectatorScreenController)
 	{
-		const bool bShouldPassLayersTexture = PipelinedLayerStateRendering.EmulatedLayerState.bIsFaceLockedLayerEmulationActive && PipelinedLayerStateRendering.EmulatedLayerState.EmulationSwapchain && !CVarOpenXRDoNotCopyEmulatedLayersToSpectatorScreen.GetValueOnRenderThread();
+		const bool bShouldPassLayersTexture = EnumHasAnyFlags(PipelinedLayerStateRendering.LayerStateFlags, EOpenXRLayerStateFlags::SubmitEmulatedFaceLockedLayer) && !CVarOpenXRDoNotCopyEmulatedLayersToSpectatorScreen.GetValueOnRenderThread();
 		const FTexture2DRHIRef LayersTexture = bShouldPassLayersTexture ? PipelinedLayerStateRendering.EmulatedLayerState.EmulationSwapchain->GetTextureRef() : nullptr;
 		SpectatorScreenController->RenderSpectatorScreen_RenderThread(RHICmdList, BackBuffer, SrcTexture, LayersTexture, WindowSize);
 	}
diff --git a/Engine/Plugins/Runtime/OpenXR/Source/OpenXRHMD/Private/OpenXRHMD.h b/Engine/Plugins/Runtime/OpenXR/Source/OpenXRHMD/Private/OpenXRHMD.h
--- a/Engine/Plugins/Runtime/OpenXR/Source/OpenXRHMD/Private/OpenXRHMD.h
+++ b/Engine/Plugins/Runtime/OpenXR/Source/OpenXRHMD/Private/OpenXRHMD.h
@@ -15,6 +15,7 @@
 #include "DefaultSpectatorScreenController.h"
 #include "IHeadMountedDisplayVulkanExtensions.h"
 #include "IOpenXRExtensionPluginDelegates.h"
+#include "Misc/EnumClassFlags.h"
 
 #include <openxr/openxr.h>
 
@@ -103,9 +104,18 @@ public:
 		TArray<XrSwapchainSubImage> EmulationImages;
 		// This swapchain is where the emulated face locked layers are rendered into.
 		FXRSwapChainPtr EmulationSwapchain;
-		bool bIsFaceLockedLayerEmulationActive = false;
 	};
 
+	enum class EOpenXRLayerStateFlags : uint32
+	{
+		None = 0u,
+		BackgroundLayerVisible = (1u << 0),
+		SubmitBackgroundLayer = (1u << 1),
+		SubmitDepthLayer = (1u << 2),
+		SubmitEmulatedFaceLockedLayer = (1u << 3),
+	};
+	FRIEND_ENUM_CLASS_FLAGS(EOpenXRLayerStateFlags);
+
 	struct FPipelinedLayerState
 	{
 		TArray<FXrCompositionLayerUnion> NativeOverlays;
@@ -121,8 +131,7 @@ public:
 
 		FEmulatedLayerState EmulatedLayerState;
 
-		bool bBackgroundLayerVisible = true;
-		bool bSubmitBackgroundLayer = true;
+		EOpenXRLayerStateFlags LayerStateFlags = EOpenXRLayerStateFlags::None;
 	};
 
 	class FVulkanExtensions : public IHeadMountedDisplayVulkanExtensions
@@ -472,3 +481,5 @@ private:
 	TArray<IStereoLayers::FLayerDesc> EmulatedFaceLockedLayers;
 	TArray<FOpenXRLayer>			  NativeLayers;
 };
+
+ENUM_CLASS_FLAGS(FOpenXRHMD::EOpenXRLayerStateFlags);
\ No newline at end of file
diff --git a/Engine/Source/Runtime/Renderer/Private/SceneTextures.cpp b/Engine/Source/Runtime/Renderer/Private/SceneTextures.cpp
--- a/Engine/Source/Runtime/Renderer/Private/SceneTextures.cpp
+++ b/Engine/Source/Runtime/Renderer/Private/SceneTextures.cpp
@@ -128,14 +128,15 @@ static IStereoRenderTargetManager* FindStereoRenderTargetManager()
 	return GEngine->StereoRenderingDevice->GetRenderTargetManager();
 }
 
-static TRefCountPtr<FRHITexture2D> FindStereoDepthTexture(uint32 bSupportsXRDepth, FIntPoint TextureExtent, uint32 NumSamples)
+static TRefCountPtr<FRHITexture2D> FindStereoDepthTexture(uint32 bSupportsXRDepth, FIntPoint TextureExtent, ETextureCreateFlags RequestedCreateFlags)
 {
 	if (bSupportsXRDepth == 1)
 	{
 		if (IStereoRenderTargetManager* StereoRenderTargetManager = FindStereoRenderTargetManager())
 		{
 			TRefCountPtr<FRHITexture2D> DepthTex, SRTex;
-			StereoRenderTargetManager->AllocateDepthTexture(0, TextureExtent.X, TextureExtent.Y, PF_DepthStencil, 1, TexCreate_None, TexCreate_DepthStencilTargetable | TexCreate_ShaderResource | TexCreate_InputAttachmentRead, DepthTex, SRTex, NumSamples);
+			constexpr uint32 NumSamples = 1;
+			StereoRenderTargetManager->AllocateDepthTexture(0, TextureExtent.X, TextureExtent.Y, PF_DepthStencil, 1, RequestedCreateFlags, TexCreate_DepthStencilTargetable | TexCreate_ShaderResource | TexCreate_InputAttachmentRead, DepthTex, SRTex, NumSamples);
 			return MoveTemp(SRTex);
 		}
 	}
@@ -442,7 +443,7 @@ void FMinimalSceneTextures::InitializeViewFamily(FRDGBuilder& GraphBuilder, FVie
 
 	// If not using MSAA, we need to make sure to grab the stereo depth texture if appropriate.
 	FTexture2DRHIRef StereoDepthRHI;
-	if (Config.NumSamples == 1 && (StereoDepthRHI = FindStereoDepthTexture(Config.bSupportsXRTargetManagerDepthAlloc, Config.Extent, Config.NumSamples)) != nullptr)
+	if (Config.NumSamples == 1 && (StereoDepthRHI = FindStereoDepthTexture(Config.bSupportsXRTargetManagerDepthAlloc, Config.Extent, ETextureCreateFlags::None)) != nullptr)
 	{
 		SceneTextures.Depth = RegisterExternalTexture(GraphBuilder, StereoDepthRHI, TEXT("SceneDepthZ"));
 		SceneTextures.Stencil = GraphBuilder.CreateSRV(FRDGTextureSRVDesc::CreateWithPixelFormat(SceneTextures.Depth.Target, PF_X24_G8));
@@ -460,7 +461,7 @@ void FMinimalSceneTextures::InitializeViewFamily(FRDGBuilder& GraphBuilder, FVie
 		{
 			Desc.NumSamples = 1;
 
-			if ((StereoDepthRHI = FindStereoDepthTexture(Config.bSupportsXRTargetManagerDepthAlloc, Config.Extent, Desc.NumSamples)) != nullptr)
+			if ((StereoDepthRHI = FindStereoDepthTexture(Config.bSupportsXRTargetManagerDepthAlloc, Config.Extent, ETextureCreateFlags::ResolveTargetable)) != nullptr)
 			{
 				ensureMsgf(Desc.ArraySize == StereoDepthRHI->GetDesc().ArraySize, TEXT("Resolve texture does not agree in dimensionality with Target (Resolve.ArraySize=%d, Target.ArraySize=%d)"),
 					Desc.ArraySize, StereoDepthRHI->GetDesc().ArraySize);
