Allow persistent headless sessions outside of VR PIE

diff --git a/Engine/Plugins/Runtime/OpenXR/Source/OpenXRHMD/Private/OpenXRHMD.cpp b/Engine/Plugins/Runtime/OpenXR/Source/OpenXRHMD/Private/OpenXRHMD.cpp
--- a/Engine/Plugins/Runtime/OpenXR/Source/OpenXRHMD/Private/OpenXRHMD.cpp
+++ b/Engine/Plugins/Runtime/OpenXR/Source/OpenXRHMD/Private/OpenXRHMD.cpp
@@ -807,6 +807,11 @@ bool FOpenXRHMD::EnableStereo(bool stereo)
 		return true;
 	}
 
+	if (bIsTrackingOnlySession)
+	{
+		return false;
+	}
+
 	bStereoEnabled = stereo;
 	if (stereo)
 	{
@@ -1256,6 +1261,7 @@ FOpenXRHMD::FOpenXRHMD(const FAutoRegister& AutoRegister, XrInstance InInstance,
 	, bIsMobileMultiViewEnabled(false)
 	, bSupportsHandTracking(false)
 	, bIsStandaloneStereoOnlyDevice(false)
+	, bIsTrackingOnlySession(false)
 	, CurrentSessionState(XR_SESSION_STATE_UNKNOWN)
 	, EnabledExtensions(std::move(InEnabledExtensions))
 	, InputModule(nullptr)
@@ -1306,6 +1312,8 @@ FOpenXRHMD::FOpenXRHMD(const FAutoRegister& AutoRegister, XrInstance InInstance,
 	}
 #endif
 
+	bIsTrackingOnlySession = FParse::Param(FCommandLine::Get(), TEXT("xrtrackingonly"));
+
 	// Add a device space for the HMD without an action handle and ensure it has the correct index
 	XrPath UserHead = XR_NULL_PATH;
 	XR_ENSURE(xrStringToPath(Instance, "/user/head", &UserHead));
@@ -1678,9 +1686,13 @@ bool FOpenXRHMD::OnStereoStartup()
 {
 	FWriteScopeLock Lock(SessionHandleMutex);
 
-	check(Session == XR_NULL_HANDLE);
 	bIsExitingSessionByxrRequestExitSession = false;  // clear in case we requested exit for a previous session, but it ended in some other way before that happened.
 
+	if (Session)
+	{
+		return false;
+	}
+
 	System = IOpenXRHMDModule::Get().GetSystemId();
 	if (!System)
 	{
@@ -2930,7 +2942,7 @@ bool FOpenXRHMD::OnStartGameFrame(FWorldContext& WorldContext)
 	// XR hardware can only be running with one of them.
 	if (GIsEditor && GEditor && GEditor->GetPIEWorldContext() != nullptr)
 	{
-		if (!WorldContext.bIsPrimaryPIEInstance)
+		if (!WorldContext.bIsPrimaryPIEInstance && !bIsTrackingOnlySession)
 		{
 			return false;
 		}
@@ -2954,6 +2966,14 @@ bool FOpenXRHMD::OnStartGameFrame(FWorldContext& WorldContext)
 		}
 	}
 
+	if (bIsTrackingOnlySession)
+	{
+		if (OnStereoStartup())
+		{
+			StartSession();
+		}
+	}
+
 	// Process all pending messages.
 	XrEventDataBuffer event;
 	while (ReadNextEvent(&event))
diff --git a/Engine/Plugins/Runtime/OpenXR/Source/OpenXRHMD/Private/OpenXRHMD.h b/Engine/Plugins/Runtime/OpenXR/Source/OpenXRHMD/Private/OpenXRHMD.h
--- a/Engine/Plugins/Runtime/OpenXR/Source/OpenXRHMD/Private/OpenXRHMD.h
+++ b/Engine/Plugins/Runtime/OpenXR/Source/OpenXRHMD/Private/OpenXRHMD.h
@@ -429,6 +429,7 @@ private:
 	bool					bSpaceAccellerationSupported;
 	bool					bProjectionLayerAlphaEnabled;
 	bool					bIsStandaloneStereoOnlyDevice;
+	bool					bIsTrackingOnlySession;
 	float					WorldToMetersScale = 100.0f;
 	float					RuntimePixelDensityMax = FHeadMountedDisplayBase::PixelDensityMax;
 
diff --git a/Engine/Plugins/Runtime/OpenXR/Source/OpenXRHMD/Private/OpenXRHMDModule.cpp b/Engine/Plugins/Runtime/OpenXR/Source/OpenXRHMD/Private/OpenXRHMDModule.cpp
--- a/Engine/Plugins/Runtime/OpenXR/Source/OpenXRHMD/Private/OpenXRHMDModule.cpp
+++ b/Engine/Plugins/Runtime/OpenXR/Source/OpenXRHMD/Private/OpenXRHMDModule.cpp
@@ -49,6 +49,11 @@ FOpenXRHMDModule::~FOpenXRHMDModule()
 
 TSharedPtr< class IXRTrackingSystem, ESPMode::ThreadSafe > FOpenXRHMDModule::CreateTrackingSystem()
 {
+	if (!InitInstance())
+	{
+		return nullptr;
+	}
+
 	if (!RenderBridge && !FParse::Param(FCommandLine::Get(), TEXT("xrtrackingonly")))
 	{
 		if (!InitRenderBridge())
