Implement HMD audio device selection for OpenXR

diff --git a/Engine/Plugins/Runtime/OpenXR/Source/OpenXRHMD/Private/OpenXRHMDModule.cpp b/Engine/Plugins/Runtime/OpenXR/Source/OpenXRHMD/Private/OpenXRHMDModule.cpp
--- a/Engine/Plugins/Runtime/OpenXR/Source/OpenXRHMD/Private/OpenXRHMDModule.cpp
+++ b/Engine/Plugins/Runtime/OpenXR/Source/OpenXRHMD/Private/OpenXRHMDModule.cpp
@@ -503,11 +503,12 @@ bool FOpenXRHMDModule::GetOptionalExtensions(TArray<const ANSICHAR*>& OutExtensi
 	OutExtensions.Add(XR_KHR_VISIBILITY_MASK_EXTENSION_NAME);
 	OutExtensions.Add(XR_KHR_BINDING_MODIFICATION_EXTENSION_NAME);
 	OutExtensions.Add(XR_EPIC_VIEW_CONFIGURATION_FOV_EXTENSION_NAME);
+	OutExtensions.Add(XR_EXT_DPAD_BINDING_EXTENSION_NAME);
 	OutExtensions.Add(XR_EXT_PALM_POSE_EXTENSION_NAME);
-
-	// Draft extension not yet provided in headers
-	OutExtensions.Add("XR_EXT_dpad_binding");
-	OutExtensions.Add("XR_EXT_active_action_set_priority");
+	OutExtensions.Add(XR_EXT_ACTIVE_ACTION_SET_PRIORITY_EXTENSION_NAME);
+#if PLATFORM_WINDOWS
+	OutExtensions.Add(XR_OCULUS_AUDIO_DEVICE_GUID_EXTENSION_NAME);
+#endif
 
 	return true;
 }
@@ -847,3 +848,37 @@ XrPath FOpenXRHMDModule::ResolveNameToPath(FName Name)
 		return XR_NULL_PATH;
 	}
 }
+
+FString FOpenXRHMDModule::GetAudioInputDevice()
+{
+#if PLATFORM_WINDOWS
+	if (Instance && IsExtensionEnabled(XR_OCULUS_AUDIO_DEVICE_GUID_EXTENSION_NAME))
+	{
+		PFN_xrGetAudioInputDeviceGuidOculus GetAudioInputDeviceGuidOculus = nullptr;
+		if (XR_ENSURE(xrGetInstanceProcAddr(Instance, "xrGetAudioInputDeviceGuidOculus", (PFN_xrVoidFunction*)&GetAudioInputDeviceGuidOculus)))
+		{
+			WCHAR DeviceGuid[XR_MAX_AUDIO_DEVICE_STR_SIZE_OCULUS];
+			GetAudioInputDeviceGuidOculus(Instance, DeviceGuid);
+			return FString(XR_MAX_AUDIO_DEVICE_STR_SIZE_OCULUS, DeviceGuid);
+		}
+	}
+#endif // PLATFORM_WINDOWS
+	return FString();
+}
+
+FString FOpenXRHMDModule::GetAudioOutputDevice()
+{
+#if PLATFORM_WINDOWS
+	if (Instance && IsExtensionEnabled(XR_OCULUS_AUDIO_DEVICE_GUID_EXTENSION_NAME))
+	{
+		PFN_xrGetAudioOutputDeviceGuidOculus GetAudioOutputDeviceGuidOculus = nullptr;
+		if (XR_ENSURE(xrGetInstanceProcAddr(Instance, "xrGetAudioOutputDeviceGuidOculus", (PFN_xrVoidFunction*)&GetAudioOutputDeviceGuidOculus)))
+		{
+			WCHAR DeviceGuid[XR_MAX_AUDIO_DEVICE_STR_SIZE_OCULUS];
+			GetAudioOutputDeviceGuidOculus(Instance, DeviceGuid);
+			return FString(XR_MAX_AUDIO_DEVICE_STR_SIZE_OCULUS, DeviceGuid);
+		}
+	}
+#endif // PLATFORM_WINDOWS
+	return FString();
+}
diff --git a/Engine/Plugins/Runtime/OpenXR/Source/OpenXRHMD/Private/OpenXRHMDModule.h b/Engine/Plugins/Runtime/OpenXR/Source/OpenXRHMD/Private/OpenXRHMDModule.h
--- a/Engine/Plugins/Runtime/OpenXR/Source/OpenXRHMD/Private/OpenXRHMDModule.h
+++ b/Engine/Plugins/Runtime/OpenXR/Source/OpenXRHMD/Private/OpenXRHMDModule.h
@@ -18,6 +18,8 @@ public:
 	virtual TSharedPtr< class IXRTrackingSystem, ESPMode::ThreadSafe > CreateTrackingSystem() override;
 	virtual TSharedPtr< class IHeadMountedDisplayVulkanExtensions, ESPMode::ThreadSafe > GetVulkanExtensions() override;
 	virtual uint64 GetGraphicsAdapterLuid() override;
+	virtual FString GetAudioInputDevice() override;
+	virtual FString GetAudioOutputDevice() override;
 
 	FString GetModuleKeyName() const override
 	{
@@ -29,11 +31,14 @@ public:
 		AliasesOut.Add(TEXT("OpenXR"));
 	}
 
-	void ShutdownModule() override;
+	virtual bool PreInit() override { return InitInstance(); }
+	virtual void ShutdownModule() override;
 
 	virtual bool IsHMDConnected() override { return true; }
 	virtual FString GetDeviceSystemName() override;
 	virtual bool IsStandaloneStereoOnlyDevice() override;
+
+	/** IOpenXRHMDModule implementation */
 	virtual bool IsExtensionAvailable(const FString& Name) const override { return AvailableExtensions.Contains(Name); }
 	virtual bool IsExtensionEnabled(const FString& Name) const override { return EnabledExtensions.Contains(Name); }
 	virtual bool IsLayerAvailable(const FString& Name) const override { return EnabledLayers.Contains(Name); }
diff --git a/Engine/Plugins/Runtime/OpenXR/Source/OpenXRInput/Private/OpenXRInput.cpp b/Engine/Plugins/Runtime/OpenXR/Source/OpenXRInput/Private/OpenXRInput.cpp
--- a/Engine/Plugins/Runtime/OpenXR/Source/OpenXRInput/Private/OpenXRInput.cpp
+++ b/Engine/Plugins/Runtime/OpenXR/Source/OpenXRInput/Private/OpenXRInput.cpp
@@ -249,8 +249,8 @@ FOpenXRInputPlugin::FOpenXRInput::FOpenXRInput(FOpenXRHMD* HMD)
 	if (OpenXRHMD)
 	{
 		Instance = OpenXRHMD->GetInstance();
-		bDirectionalBindingSupported = OpenXRHMD->IsExtensionEnabled("XR_EXT_dpad_binding");
-		bPalmPoseSupported = OpenXRHMD->IsExtensionEnabled("XR_EXT_palm_pose");
+		bDirectionalBindingSupported = OpenXRHMD->IsExtensionEnabled(XR_EXT_DPAD_BINDING_EXTENSION_NAME);
+		bPalmPoseSupported = OpenXRHMD->IsExtensionEnabled(XR_EXT_PALM_POSE_EXTENSION_NAME);
 
 		// Note: AnyHand needs special handling because it tries left then falls back to right in each call.
 		MotionSourceToControllerHandMap.Add(OpenXRSourceNames::Left, EControllerHand::Left);
