Allow acquiring on the RenderThread if the runtime supports it

diff --git a/Engine/Plugins/Runtime/OpenXR/Source/OpenXRHMD/Private/OpenXRHMD.cpp b/Engine/Plugins/Runtime/OpenXR/Source/OpenXRHMD/Private/OpenXRHMD.cpp
--- a/Engine/Plugins/Runtime/OpenXR/Source/OpenXRHMD/Private/OpenXRHMD.cpp
+++ b/Engine/Plugins/Runtime/OpenXR/Source/OpenXRHMD/Private/OpenXRHMD.cpp
@@ -1245,6 +1245,15 @@ bool CheckPlatformDepthExtensionSupport(const XrInstanceProperties& InstanceProp
 	return true;
 }
 
+bool CheckPlatformAcquireOnRenderThreadSupport(const XrInstanceProperties& InstanceProps)
+{
+	if (RHIGetInterfaceType() != ERHIInterfaceType::Vulkan || FCStringAnsi::Strstr(InstanceProps.runtimeName, "Oculus"))
+	{
+		return true;
+	}
+	return false;
+}
+
 FOpenXRHMD::FOpenXRHMD(const FAutoRegister& AutoRegister, XrInstance InInstance, TRefCountPtr<FOpenXRRenderBridge>& InRenderBridge, TArray<const char*> InEnabledExtensions, TArray<IOpenXRExtensionPlugin*> InExtensionPlugins, IARSystemSupport* ARSystemSupport)
 	: FHeadMountedDisplayBase(ARSystemSupport)
 	, FHMDSceneViewExtension(AutoRegister)
@@ -1295,7 +1304,7 @@ FOpenXRHMD::FOpenXRHMD(const FAutoRegister& AutoRegister, XrInstance InInstance,
 	bViewConfigurationFovSupported = IsExtensionEnabled(XR_EPIC_VIEW_CONFIGURATION_FOV_EXTENSION_NAME);
 	bSupportsHandTracking = IsExtensionEnabled(XR_EXT_HAND_TRACKING_EXTENSION_NAME);
 	bSpaceAccellerationSupported = IsExtensionEnabled(XR_EPIC_SPACE_ACCELERATION_NAME);
-
+	bIsAcquireOnRenderThreadSupported = CheckPlatformAcquireOnRenderThreadSupport(InstanceProperties);
 	ReconfigureForShaderPlatform(GMaxRHIShaderPlatform);
 
 #if PLATFORM_HOLOLENS || PLATFORM_ANDROID
@@ -2757,14 +2766,33 @@ void FOpenXRHMD::OnBeginRendering_RenderThread(FRHICommandListImmediate& RHICmdL
 		// Always check for sanity before using it.
 		FXRSwapChainPtr EmulationSwapchain = PipelinedLayerStateRendering.EmulatedLayerState.EmulationSwapchain;
 
+		if (bIsAcquireOnRenderThreadSupported)
+		{
+			if (ColorSwapchain)
+			{
+				ColorSwapchain->IncrementSwapChainIndex_RHIThread();
+				if (bDepthExtensionSupported && DepthSwapchain)
+				{
+					DepthSwapchain->IncrementSwapChainIndex_RHIThread();
+				}
+				if (PipelinedLayerStateRHI.EmulatedLayerState.bIsFaceLockedLayerEmulationActive && EmulationSwapchain)
+				{
+					EmulationSwapchain->IncrementSwapChainIndex_RHIThread();
+				}
+			}
+		}
+
 		RHICmdList.EnqueueLambda([this, FrameState = PipelinedFrameStateRendering, ColorSwapchain, DepthSwapchain, EmulationSwapchain](FRHICommandListImmediate& InRHICmdList)
 		{
 			OnBeginRendering_RHIThread(FrameState, ColorSwapchain, DepthSwapchain, EmulationSwapchain);
 		});
 
-		// We need to sync with the RHI thread to ensure we've acquired the next swapchain image.
-		// TODO: The acquire needs to be moved to the Render thread as soon as it's allowed by the spec.
-		RHICmdList.ImmediateFlush(EImmediateFlushType::FlushRHIThread);
+		if (!bIsAcquireOnRenderThreadSupported)
+		{
+			// We need to sync with the RHI thread to ensure we've acquired the next swapchain image.
+			// TODO: The acquire needs to be moved to the Render thread as soon as it's allowed by the spec.
+			RHICmdList.ImmediateFlush(EImmediateFlushType::FlushRHIThread);
+		}
 	}
 
 	// Snapshot new poses for late update.
@@ -3174,16 +3202,25 @@ void FOpenXRHMD::OnBeginRendering_RHIThread(const FPipelinedFrameState& InFrameS
 		// We need a new swapchain image unless we've already acquired one for rendering
 		if (!bIsRendering && ColorSwapchain)
 		{
-			ColorSwapchain->IncrementSwapChainIndex_RHIThread();
+			if (!bIsAcquireOnRenderThreadSupported)
+			{
+				ColorSwapchain->IncrementSwapChainIndex_RHIThread();
+			}
 			ColorSwapchain->WaitCurrentImage_RHIThread(OPENXR_SWAPCHAIN_WAIT_TIMEOUT);
 			if (DepthSwapchain)
 			{
-				DepthSwapchain->IncrementSwapChainIndex_RHIThread();
+				if (!bIsAcquireOnRenderThreadSupported)
+				{
+					DepthSwapchain->IncrementSwapChainIndex_RHIThread();
+				}
 				DepthSwapchain->WaitCurrentImage_RHIThread(OPENXR_SWAPCHAIN_WAIT_TIMEOUT);
 			}
 			if (EmulationSwapchain)
 			{
-				EmulationSwapchain->IncrementSwapChainIndex_RHIThread();
+				if (!bIsAcquireOnRenderThreadSupported)
+				{
+					EmulationSwapchain->IncrementSwapChainIndex_RHIThread();
+				}
 				EmulationSwapchain->WaitCurrentImage_RHIThread(OPENXR_SWAPCHAIN_WAIT_TIMEOUT);
 			}
 		}
diff --git a/Engine/Plugins/Runtime/OpenXR/Source/OpenXRHMD/Private/OpenXRHMD.h b/Engine/Plugins/Runtime/OpenXR/Source/OpenXRHMD/Private/OpenXRHMD.h
--- a/Engine/Plugins/Runtime/OpenXR/Source/OpenXRHMD/Private/OpenXRHMD.h
+++ b/Engine/Plugins/Runtime/OpenXR/Source/OpenXRHMD/Private/OpenXRHMD.h
@@ -430,6 +430,7 @@ private:
 	bool					bProjectionLayerAlphaEnabled;
 	bool					bIsStandaloneStereoOnlyDevice;
 	bool					bIsTrackingOnlySession;
+	bool					bIsAcquireOnRenderThreadSupported;
 	float					WorldToMetersScale = 100.0f;
 	float					RuntimePixelDensityMax = FHeadMountedDisplayBase::PixelDensityMax;
 
