Fix merge error in OpenXRHMD.cpp

diff --git a/Engine/Plugins/Runtime/OpenXR/Source/OpenXRHMD/Private/OpenXRHMD.cpp b/Engine/Plugins/Runtime/OpenXR/Source/OpenXRHMD/Private/OpenXRHMD.cpp
--- a/Engine/Plugins/Runtime/OpenXR/Source/OpenXRHMD/Private/OpenXRHMD.cpp
+++ b/Engine/Plugins/Runtime/OpenXR/Source/OpenXRHMD/Private/OpenXRHMD.cpp
@@ -2775,7 +2775,7 @@ void FOpenXRHMD::OnBeginRendering_RenderThread(FRHICommandListImmediate& RHICmdL
 				{
 					DepthSwapchain->IncrementSwapChainIndex_RHIThread();
 				}
-				if (PipelinedLayerStateRHI.EmulatedLayerState.bIsFaceLockedLayerEmulationActive && EmulationSwapchain)
+				if (EmulationSwapchain)
 				{
 					EmulationSwapchain->IncrementSwapChainIndex_RHIThread();
 				}
