Fix Android VR Preview issues (UE-176854).

- Prevent wrong feature level from being used while in preview.
- Also while at it, remove the filtering out of a alpha blend mode on mobile, because mobile renderer no longer uses alpha to store depth.

diff --git a/Engine/Plugins/Runtime/OpenXR/Source/OpenXRHMD/Private/OpenXRHMD.cpp b/Engine/Plugins/Runtime/OpenXR/Source/OpenXRHMD/Private/OpenXRHMD.cpp
--- a/Engine/Plugins/Runtime/OpenXR/Source/OpenXRHMD/Private/OpenXRHMD.cpp
+++ b/Engine/Plugins/Runtime/OpenXR/Source/OpenXRHMD/Private/OpenXRHMD.cpp
@@ -78,21 +78,8 @@ static TAutoConsoleVariable<bool> CVarOpenXRDoNotCopyEmulatedLayersToSpectatorSc
 	ECVF_Default);
 
 namespace {
-	static TSet<XrEnvironmentBlendMode> SupportedBlendModes{ XR_ENVIRONMENT_BLEND_MODE_ALPHA_BLEND, XR_ENVIRONMENT_BLEND_MODE_ADDITIVE, XR_ENVIRONMENT_BLEND_MODE_OPAQUE };
 	static TSet<XrViewConfigurationType> SupportedViewConfigurations{ XR_VIEW_CONFIGURATION_TYPE_PRIMARY_STEREO, XR_VIEW_CONFIGURATION_TYPE_PRIMARY_QUAD_VARJO };
 
-	/** Helper function for checking whether a blend mode is supported. */
-	bool IsEnvironmentBlendModeSupported(XrEnvironmentBlendMode BlendMode)
-	{
-		if (SupportedBlendModes.Contains(BlendMode) &&
-			// On mobile platforms the alpha channel can contain depth information, so we can't use alpha blend.
-			(BlendMode != XR_ENVIRONMENT_BLEND_MODE_ALPHA_BLEND || !IsMobilePlatform(GMaxRHIShaderPlatform)))
-		{
-			return true;
-		}
-		return false;
-	}
-
 	/** Helper function for acquiring the appropriate FSceneViewport */
 	FSceneViewport* FindSceneViewport()
 	{
@@ -1347,6 +1334,14 @@ bool FOpenXRHMD::ReconfigureForShaderPlatform(EShaderPlatform NewShaderPlatform)
 	static const auto CVarPropagateAlpha = IConsoleManager::Get().FindTConsoleVariableDataInt(TEXT("r.PostProcessing.PropagateAlpha"));
 	bProjectionLayerAlphaEnabled = !IsMobilePlatform(NewShaderPlatform) && CVarPropagateAlpha->GetValueOnAnyThread() != 0;
 
+	ConfiguredShaderPlatform = NewShaderPlatform;
+
+	UE_LOG(LogHMD, Log, TEXT("HMD configured for shader platform %s, bIsMobileMultiViewEnabled=%d, bProjectionLayerAlphaEnabled=%d"),
+		*LexToString(ConfiguredShaderPlatform),
+		bIsMobileMultiViewEnabled,
+		bProjectionLayerAlphaEnabled
+		);
+
 	return true;
 }
 
@@ -1756,23 +1751,10 @@ bool FOpenXRHMD::OnStereoStartup()
 		RuntimePixelDensityMax = FMath::Min(RuntimePixelDensityMax, PerViewPixelDensityMax);
 	}
 
-	// Enumerate environment blend modes and select the best one.
+	// Select the first blend mode returned by the runtime - as per spec, environment blend modes should be in order from highest to lowest runtime preference
 	{
 		TArray<XrEnvironmentBlendMode> BlendModes = RetrieveEnvironmentBlendModes();
-
-		// Select the first blend mode returned by the runtime that is supported.
-		// This is the environment blend mode preferred by the runtime.
-		for (XrEnvironmentBlendMode BlendMode : BlendModes)
-		{
-			if (IsEnvironmentBlendModeSupported(BlendMode))
-			{
-				SelectedEnvironmentBlendMode = BlendMode;
-				break;
-			}
-		}
-
-		// If there is no supported environment blend mode, use the first option as a last resort.
-		if (!ensure(SelectedEnvironmentBlendMode != XR_ENVIRONMENT_BLEND_MODE_MAX_ENUM))
+		if (!BlendModes.IsEmpty())
 		{
 			SelectedEnvironmentBlendMode = BlendModes[0];
 		}
@@ -2141,7 +2123,7 @@ void FOpenXRHMD::SetEnvironmentBlendMode(XrEnvironmentBlendMode NewBlendMode)
 
 	TArray<XrEnvironmentBlendMode> BlendModes = RetrieveEnvironmentBlendModes();
 
-	if (BlendModes.Contains(NewBlendMode) && IsEnvironmentBlendModeSupported(NewBlendMode))
+	if (BlendModes.Contains(NewBlendMode))
 	{
 		SelectedEnvironmentBlendMode = NewBlendMode;
 		UE_LOG(LogHMD, Log, TEXT("Environment Blend Mode set to: %d."), SelectedEnvironmentBlendMode);
@@ -2228,7 +2210,7 @@ bool FOpenXRHMD::AllocateRenderTargetTexture(uint32 Index, uint32 SizeX, uint32
 	UnifiedCreateFlags |= TexCreate_RenderTargetable;
 
 	// On mobile without HDR all render targets need to be marked sRGB
-	bool MobileHWsRGB = IsMobileColorsRGB() && IsMobilePlatform(GMaxRHIShaderPlatform);
+	bool MobileHWsRGB = IsMobileColorsRGB() && IsMobilePlatform(GetConfiguredShaderPlatform());
 	if (MobileHWsRGB)
 	{
 		UnifiedCreateFlags |= TexCreate_SRGB;
@@ -3662,8 +3644,7 @@ void FOpenXRHMD::CopyTexture_RenderThread(FRHICommandListImmediate& RHICmdList,
 		GraphicsPSOInit.DepthStencilState = TStaticDepthStencilState<false, CF_Always>::GetRHI();
 		GraphicsPSOInit.PrimitiveType = PT_TriangleList;
 
-		const auto FeatureLevel = GMaxRHIFeatureLevel;
-		FGlobalShaderMap* ShaderMap = GetGlobalShaderMap(FeatureLevel);
+		FGlobalShaderMap* ShaderMap = GetGlobalShaderMap(GetConfiguredShaderPlatform());
 
 		TShaderMapRef<FScreenVS> VertexShader(ShaderMap);
 
@@ -3693,7 +3674,7 @@ void FOpenXRHMD::CopyTexture_RenderThread(FRHICommandListImmediate& RHICmdList,
 			}
 		}
 
-		bNeedsDisplayMapping &= (GMaxRHIFeatureLevel > ERHIFeatureLevel::ES3_1);
+		bNeedsDisplayMapping &= IsFeatureLevelSupported(GetConfiguredShaderPlatform(), ERHIFeatureLevel::ES3_1);
 
 		bool bIsArraySource = SrcTexture->GetDesc().IsTextureArray();
 
diff --git a/Engine/Plugins/Runtime/OpenXR/Source/OpenXRHMD/Private/OpenXRHMD.h b/Engine/Plugins/Runtime/OpenXR/Source/OpenXRHMD/Private/OpenXRHMD.h
--- a/Engine/Plugins/Runtime/OpenXR/Source/OpenXRHMD/Private/OpenXRHMD.h
+++ b/Engine/Plugins/Runtime/OpenXR/Source/OpenXRHMD/Private/OpenXRHMD.h
@@ -401,6 +401,9 @@ public:
 	OPENXRHMD_API TArray<IOpenXRExtensionPlugin*>& GetExtensionPlugins() { return ExtensionPlugins; }
 	OPENXRHMD_API void SetEnvironmentBlendMode(XrEnvironmentBlendMode NewBlendMode);
 
+	/** Returns shader platform the plugin is currently configured for, in the editor it can change due to preview platforms. */
+	EShaderPlatform GetConfiguredShaderPlatform() const { check(ConfiguredShaderPlatform != EShaderPlatform::SP_NumPlatforms); return ConfiguredShaderPlatform; }
+
 private:
 
 	TArray<XrEnvironmentBlendMode> RetrieveEnvironmentBlendModes() const;
@@ -433,6 +436,7 @@ private:
 	bool					bIsAcquireOnRenderThreadSupported;
 	float					WorldToMetersScale = 100.0f;
 	float					RuntimePixelDensityMax = FHeadMountedDisplayBase::PixelDensityMax;
+	EShaderPlatform			ConfiguredShaderPlatform = EShaderPlatform::SP_NumPlatforms;
 
 	XrSessionState			CurrentSessionState;
 	FRWLock					SessionHandleMutex;
diff --git a/Engine/Plugins/Runtime/OpenXR/Source/OpenXRHMD/Private/OpenXRHMD_Swapchain.cpp b/Engine/Plugins/Runtime/OpenXR/Source/OpenXRHMD/Private/OpenXRHMD_Swapchain.cpp
--- a/Engine/Plugins/Runtime/OpenXR/Source/OpenXRHMD/Private/OpenXRHMD_Swapchain.cpp
+++ b/Engine/Plugins/Runtime/OpenXR/Source/OpenXRHMD/Private/OpenXRHMD_Swapchain.cpp
@@ -320,7 +320,10 @@ FXRSwapChainPtr CreateSwapchain_D3D12(XrSession InSession, uint8 Format, uint8&
 	TArray<XrSwapchainImageD3D12KHR> Images = EnumerateImages<XrSwapchainImageD3D12KHR>(Swapchain, XR_TYPE_SWAPCHAIN_IMAGE_D3D12_KHR);
 	for (const auto& Image : Images)
 	{
-		TextureChain.Add(static_cast<FTextureRHIRef>(DynamicRHI->RHICreateTexture2DArrayFromResource(GPixelFormats[Format].UnrealFormat, CreateFlags, ClearValueBinding, Image.texture)));
+		TextureChain.Add(static_cast<FTextureRHIRef>((ArraySize > 1) ?
+			DynamicRHI->RHICreateTexture2DArrayFromResource(GPixelFormats[Format].UnrealFormat, CreateFlags, ClearValueBinding, Image.texture) :
+			DynamicRHI->RHICreateTexture2DFromResource(GPixelFormats[Format].UnrealFormat, CreateFlags, ClearValueBinding, Image.texture)
+			));
 	}
 	FTextureRHIRef ChainTarget = static_cast<FTextureRHIRef>(DynamicRHI->RHICreateAliasedTexture((FTextureRHIRef&)TextureChain[0]));
 
diff --git a/Engine/Source/Editor/UnrealEd/Private/EditorEngine.cpp b/Engine/Source/Editor/UnrealEd/Private/EditorEngine.cpp
--- a/Engine/Source/Editor/UnrealEd/Private/EditorEngine.cpp
+++ b/Engine/Source/Editor/UnrealEd/Private/EditorEngine.cpp
@@ -7801,7 +7801,8 @@ void UEditorEngine::OnEffectivePreviewShaderPlatformChange()
 		IStereoRenderTargetManager* StereoRenderTargetManager = StereoRenderingDevice->GetRenderTargetManager();
 		if (StereoRenderTargetManager)
 		{
-			StereoRenderTargetManager->ReconfigureForShaderPlatform(PreviewPlatform.ShaderPlatform);
+			StereoRenderTargetManager->ReconfigureForShaderPlatform(
+				PreviewPlatform.bPreviewFeatureLevelActive ? PreviewPlatform.ShaderPlatform : CachedEditorShaderPlatform);
 		}
 	}
 }
