UE-180630 FOpenXRHMD::UpdateLayer thread safety
-This function works on the NativeQuadLayers, which are RenderThread data.  It ought to enqueue this work.

diff --git a/Engine/Plugins/Runtime/OpenXR/Source/OpenXRHMD/Private/OpenXRHMD.cpp b/Engine/Plugins/Runtime/OpenXR/Source/OpenXRHMD/Private/OpenXRHMD.cpp
--- a/Engine/Plugins/Runtime/OpenXR/Source/OpenXRHMD/Private/OpenXRHMD.cpp
+++ b/Engine/Plugins/Runtime/OpenXR/Source/OpenXRHMD/Private/OpenXRHMD.cpp
@@ -3825,26 +3825,36 @@ void FOpenXRHMD::DrawVisibleAreaMesh(class FRHICommandList& RHICmdList, int32 Vi
 
 void FOpenXRHMD::UpdateLayer(FOpenXRLayer& ManagerLayer, uint32 LayerId, bool bIsValid)
 {
-	for (FOpenXRLayer& NativeLayer : NativeLayers)
+	ENQUEUE_RENDER_COMMAND(UpdateLayer)(
+		[
+			this,
+			Flags = ManagerLayer.Desc.Flags,
+			bUpdateRightEyeTexture = ManagerLayer.RightEye.bUpdateTexture,
+			bUpdateLeftEyeTexture = ManagerLayer.LeftEye.bUpdateTexture,
+			LayerId
+		](FRHICommandList&)
 	{
-		if (NativeLayer.GetLayerId() == LayerId)
+	for (FOpenXRLayer& NativeLayer : NativeLayers)
 		{
-			const bool bStaticSwapchain = !(ManagerLayer.Desc.Flags & IStereoLayers::LAYER_FLAG_TEX_CONTINUOUS_UPDATE);
-			NativeLayer.RightEye.bUpdateTexture = ManagerLayer.RightEye.bUpdateTexture;
-			NativeLayer.LeftEye.bUpdateTexture = ManagerLayer.LeftEye.bUpdateTexture;
-			if (bStaticSwapchain)
+			if (NativeLayer.GetLayerId() == LayerId)
 			{
-				if (NativeLayer.RightEye.bUpdateTexture) 
+				const bool bStaticSwapchain = !(Flags & IStereoLayers::LAYER_FLAG_TEX_CONTINUOUS_UPDATE);
+				NativeLayer.RightEye.bUpdateTexture = bUpdateRightEyeTexture;
+				NativeLayer.LeftEye.bUpdateTexture = bUpdateLeftEyeTexture;
+				if (bStaticSwapchain)
 				{
-					NativeLayer.RightEye.Swapchain.Reset();
-				}
-				if (NativeLayer.LeftEye.bUpdateTexture)
-				{
-					NativeLayer.LeftEye.Swapchain.Reset();
+					if (NativeLayer.RightEye.bUpdateTexture)
+					{
+						NativeLayer.RightEye.Swapchain.Reset();
+					}
+					if (NativeLayer.LeftEye.bUpdateTexture)
+					{
+						NativeLayer.LeftEye.Swapchain.Reset();
+					}
 				}
 			}
 		}
-	}
+	});
 }
 
 //---------------------------------------------------
