Fix packaged Quest loading hang due to recent changes in FAndroidWindow::GetScreenRect and the OpenXR depth swapchain view rect initialization.
The FAndroidWindow::GetScreenRect change left us with a (0,0) resolution when using Quest, so the HMD wouldn't even try to render.
The depth swapchain changes cause the depth swapchain ViewRect to be uninitialized, resulting in an XR_ERROR_SWAPCHAIN_RECT_INVALID error.

diff --git a/Engine/Plugins/Runtime/OpenXR/Source/OpenXRHMD/Private/OpenXRHMD.cpp b/Engine/Plugins/Runtime/OpenXR/Source/OpenXRHMD/Private/OpenXRHMD.cpp
--- a/Engine/Plugins/Runtime/OpenXR/Source/OpenXRHMD/Private/OpenXRHMD.cpp
+++ b/Engine/Plugins/Runtime/OpenXR/Source/OpenXRHMD/Private/OpenXRHMD.cpp
@@ -952,17 +952,12 @@ void FOpenXRHMD::SetFinalViewRect(FRHICommandListImmediate& RHICmdList, const in
 	};
 
 	XrSwapchainSubImage& DepthImage = PipelinedLayerStateRendering.DepthImages[ViewIndex];
-	if (EnumHasAnyFlags(PipelinedLayerStateRendering.LayerStateFlags, EOpenXRLayerStateFlags::SubmitDepthLayer))
-	{
-		DepthImage.imageArrayIndex = ColorImage.imageArrayIndex;
-		DepthImage.imageRect = ColorImage.imageRect;
-	}
+	DepthImage.imageArrayIndex = ColorImage.imageArrayIndex;
+	DepthImage.imageRect = ColorImage.imageRect;
+
 	XrSwapchainSubImage& EmulationImage = PipelinedLayerStateRendering.EmulatedLayerState.EmulationImages[ViewIndex];
-	if (EnumHasAnyFlags(PipelinedLayerStateRendering.LayerStateFlags, EOpenXRLayerStateFlags::SubmitEmulatedFaceLockedLayer))
-	{
-		EmulationImage.imageArrayIndex = ColorImage.imageArrayIndex;
-		EmulationImage.imageRect = ColorImage.imageRect;
-	}
+	EmulationImage.imageArrayIndex = ColorImage.imageArrayIndex;
+	EmulationImage.imageRect = ColorImage.imageRect;
 }
 
 EStereoscopicPass FOpenXRHMD::GetViewPassForIndex(bool bStereoRequested, int32 ViewIndex) const
