OpenXR: Fix build break from 25001160

Some legacy acquire code was left over from the merge, and invalid after the changes in 25001160
Original CL is from UE/5.2: 25000940

diff --git a/Engine/Plugins/Runtime/OpenXR/Source/OpenXRHMD/Private/OpenXRHMD.cpp b/Engine/Plugins/Runtime/OpenXR/Source/OpenXRHMD/Private/OpenXRHMD.cpp
--- a/Engine/Plugins/Runtime/OpenXR/Source/OpenXRHMD/Private/OpenXRHMD.cpp
+++ b/Engine/Plugins/Runtime/OpenXR/Source/OpenXRHMD/Private/OpenXRHMD.cpp
@@ -1233,15 +1233,6 @@ bool CheckPlatformDepthExtensionSupport(const XrInstanceProperties& InstanceProp
 	return true;
 }
 
-bool CheckPlatformAcquireOnRenderThreadSupport(const XrInstanceProperties& InstanceProps)
-{
-	if (RHIGetInterfaceType() != ERHIInterfaceType::Vulkan || FCStringAnsi::Strstr(InstanceProps.runtimeName, "Oculus"))
-	{
-		return true;
-	}
-	return false;
-}
-
 bool CheckPlatformAcquireOnAnyThreadSupport(const XrInstanceProperties& InstanceProps)
 {
 	int32 AcquireMode = CVarOpenXRAcquireMode.GetValueOnAnyThread();
@@ -2803,33 +2794,10 @@ void FOpenXRHMD::OnBeginRendering_RenderThread(FRHICommandListImmediate& RHICmdL
 		// Always check for sanity before using it.
 		FXRSwapChainPtr EmulationSwapchain = PipelinedLayerStateRendering.EmulatedLayerState.EmulationSwapchain;
 
-		if (bIsAcquireOnRenderThreadSupported)
-		{
-			if (ColorSwapchain)
-			{
-				ColorSwapchain->IncrementSwapChainIndex_RHIThread();
-				if (bDepthExtensionSupported && DepthSwapchain)
-				{
-					DepthSwapchain->IncrementSwapChainIndex_RHIThread();
-				}
-				if (EmulationSwapchain)
-				{
-					EmulationSwapchain->IncrementSwapChainIndex_RHIThread();
-				}
-			}
-		}
-
 		RHICmdList.EnqueueLambda([this, FrameState = PipelinedFrameStateRendering, ColorSwapchain, DepthSwapchain, EmulationSwapchain](FRHICommandListImmediate& InRHICmdList)
 		{
 			OnBeginRendering_RHIThread(FrameState, ColorSwapchain, DepthSwapchain, EmulationSwapchain);
 		});
-
-		if (!bIsAcquireOnRenderThreadSupported)
-		{
-			// We need to sync with the RHI thread to ensure we've acquired the next swapchain image.
-			// TODO: The acquire needs to be moved to the Render thread as soon as it's allowed by the spec.
-			RHICmdList.ImmediateFlush(EImmediateFlushType::FlushRHIThread);
-		}
 	}
 
 	// Snapshot new poses for late update.
