PR #10304: Add interface within IOpenXRExtensionPlugin.h to Access StereoLayers (Contributed by [REDACTED_PERSON])

diff --git a/Engine/Plugins/Runtime/OpenXR/Source/OpenXRHMD/Private/OpenXRHMD.cpp b/Engine/Plugins/Runtime/OpenXR/Source/OpenXRHMD/Private/OpenXRHMD.cpp
--- a/Engine/Plugins/Runtime/OpenXR/Source/OpenXRHMD/Private/OpenXRHMD.cpp
+++ b/Engine/Plugins/Runtime/OpenXR/Source/OpenXRHMD/Private/OpenXRHMD.cpp
@@ -2439,6 +2439,8 @@ bool FOpenXRHMD::IsEmulatingStereoLayers()
 void FOpenXRHMD::SetupFrameLayers_RenderThread(FRHICommandListImmediate& RHICmdList)
 {
 	ensure(IsInRenderingThread());
+
+	TArray<uint32> LayerIds;
 	if (GetStereoLayersDirty())
 	{
 		// When SetupFrameLayers_RenderThread is called more than once with GetStereoLayersDirty = true,
@@ -2454,8 +2456,10 @@ void FOpenXRHMD::SetupFrameLayers_RenderThread(FRHICommandListImmediate& RHICmdL
 		NativeLayers.Reset();
 
 		// Go over the dirtied layers to bin them into either native or emulated
-		ForEachLayer([&](uint32 /* unused */, FOpenXRLayer& Layer)
+		ForEachLayer([&](uint32 LayerId, FOpenXRLayer& Layer)
 		{
+			LayerIds.Add(LayerId);
+
 			if (IsEmulatingStereoLayers())
 			{
 				// Only quad layers are supported by emulation.
@@ -2513,6 +2517,15 @@ void FOpenXRHMD::SetupFrameLayers_RenderThread(FRHICommandListImmediate& RHICmdL
 		BackgroundCompositedEmulatedLayers.Sort(LayerCompare);
 		EmulatedFaceLockedLayers.Sort(LayerCompare);
 		NativeLayers.Sort(LayerCompare);
+	} //GetStereoLayersDirty()
+	else
+	{
+		// If GetStereoLayersDirty() is false, we still need to recover
+		// all the layer ids to allow plugins to access them in OnSetupLayers_RenderThread.
+		ForEachLayer([&LayerIds](uint32 LayerId, FOpenXRLayer& Layer)
+		{
+			LayerIds.Add(LayerId);
+		});
 	}
 
 	PipelinedLayerStateRendering.LayerStateFlags |= !EmulatedFaceLockedLayers.IsEmpty() ? EOpenXRLayerStateFlags::SubmitEmulatedFaceLockedLayer : EOpenXRLayerStateFlags::None;
@@ -2540,6 +2553,11 @@ void FOpenXRHMD::SetupFrameLayers_RenderThread(FRHICommandListImmediate& RHICmdL
 		PipelinedLayerStateRendering.NativeOverlays.Append(Headers);
 		UpdateLayerSwapchainTexture(Layer, RHICmdList);
 	}
+
+	for (IOpenXRExtensionPlugin* Module : ExtensionPlugins)
+	{
+		Module->OnSetupLayers_RenderThread(Session, LayerIds);
+	}
 }
 
 void FOpenXRHMD::ConfigureLayerSwapchain(FOpenXRLayer& Layer, TArray<FOpenXRLayer>& BackupLayers)
@@ -3360,6 +3378,11 @@ void FOpenXRHMD::AddLayersToHeaders(TArray<const XrCompositionLayerBaseHeader*>&
 	{
 		Headers.Add(reinterpret_cast<const XrCompositionLayerBaseHeader*>(&Layer.Header));
 	}
+
+	for (IOpenXRExtensionPlugin* Module : ExtensionPlugins)
+	{
+		Module->UpdateCompositionLayers(Session, Headers);
+	}
 }
 
 FXRRenderBridge* FOpenXRHMD::GetActiveRenderBridge_GameThread(bool /* bUseSeparateRenderTarget */)
diff --git a/Engine/Plugins/Runtime/OpenXR/Source/OpenXRHMD/Public/IOpenXRExtensionPlugin.h b/Engine/Plugins/Runtime/OpenXR/Source/OpenXRHMD/Public/IOpenXRExtensionPlugin.h
--- a/Engine/Plugins/Runtime/OpenXR/Source/OpenXRHMD/Public/IOpenXRExtensionPlugin.h
+++ b/Engine/Plugins/Runtime/OpenXR/Source/OpenXRHMD/Public/IOpenXRExtensionPlugin.h
@@ -375,6 +375,16 @@ public:
 	virtual const void* OnBeginDepthInfo(XrSession InSession, int32 InLayerIndex, int32 InViewIndex, const void* InNext)
 	{
 		return InNext;
+	}	
+	
+	// FOpenXRHMD::OnBeginRendering_RenderThread
+	virtual void OnSetupLayers_RenderThread(XrSession InSession, const TArray<uint32>& LayerIds)
+	{
+	}
+
+	// FOpenXRHMD::OnFinishRendering_RHIThread
+	virtual void UpdateCompositionLayers(XrSession InSession, TArray<const XrCompositionLayerBaseHeader*>& Headers)
+	{
 	}
 
 	virtual const void* OnEndProjectionLayer(XrSession InSession, int32 InLayerIndex, const void* InNext, XrCompositionLayerFlags& OutFlags)
