Enable setting a preferred OpenXR configuration to override the runtime's own preference. Allows (for instance) disabling Varjo quad-view.

diff --git a/Engine/Plugins/Runtime/OpenXR/Source/OpenXRHMD/Private/OpenXRHMD.cpp b/Engine/Plugins/Runtime/OpenXR/Source/OpenXRHMD/Private/OpenXRHMD.cpp
--- a/Engine/Plugins/Runtime/OpenXR/Source/OpenXRHMD/Private/OpenXRHMD.cpp
+++ b/Engine/Plugins/Runtime/OpenXR/Source/OpenXRHMD/Private/OpenXRHMD.cpp
@@ -83,6 +83,13 @@ static TAutoConsoleVariable<int32> CVarOpenXRAcquireMode(
 	TEXT("Override the swapchain acquire mode. 1 = Acquire on any thread, 2 = Only acquire on RHI thread\n"),
 	ECVF_Default);
 
+static TAutoConsoleVariable<int32> CVarOpenXRPreferredViewConfiguration(
+	TEXT("xr.OpenXRPreferredViewConfiguration"),
+	0,
+	TEXT("Override the runtime's preferred view configuration if the selected configuration is available.\n")
+	TEXT("1 = Mono, 2 = Stereo\n"),
+	ECVF_Default);
+
 namespace {
 	static TSet<XrViewConfigurationType> SupportedViewConfigurations{ XR_VIEW_CONFIGURATION_TYPE_PRIMARY_STEREO, XR_VIEW_CONFIGURATION_TYPE_PRIMARY_QUAD_VARJO };
 
@@ -1724,16 +1731,27 @@ bool FOpenXRHMD::OnStereoStartup()
 	// Fill the initial array with valid enum types (this will fail in the validation layer otherwise).
 	ViewConfigTypes.Init(XR_VIEW_CONFIGURATION_TYPE_PRIMARY_MONO, ConfigurationCount);
 	XR_ENSURE(xrEnumerateViewConfigurations(Instance, System, ConfigurationCount, &ConfigurationCount, ViewConfigTypes.GetData()));
+	// Filter to supported configurations only
+	ViewConfigTypes = ViewConfigTypes.FilterByPredicate([&](XrViewConfigurationType Type) 
+		{ 
+			return SupportedViewConfigurations.Contains(Type); 
+		});
 
-	// Select the first view configuration returned by the runtime that is supported.
+	// If we've specified a view configuration override and it's available, try to use that.
+	// Otherwise select the first view configuration returned by the runtime that is supported.
 	// This is the view configuration preferred by the runtime.
-	for (XrViewConfigurationType ViewConfigType : ViewConfigTypes)
-	{
-		if (SupportedViewConfigurations.Contains(ViewConfigType))
+	XrViewConfigurationType* PreferredViewConfiguration = ViewConfigTypes.FindByPredicate([&](XrViewConfigurationType Type)
 		{
-			SelectedViewConfigurationType = ViewConfigType;
-			break;
-		}
+			return Type == CVarOpenXRPreferredViewConfiguration.GetValueOnGameThread();
+		});
+
+	if (PreferredViewConfiguration)
+	{
+		SelectedViewConfigurationType = *PreferredViewConfiguration;
+	}
+	else if (ViewConfigTypes.Num() > 0)
+	{
+		SelectedViewConfigurationType = ViewConfigTypes[0];
 	}
 
 	// If there is no supported view configuration type, use the first option as a last resort.
