Add support for XR_FB_composition_layer_alpha_blend extension to allow for alpha inversion when submitting layers to openxr compositor.

diff --git a/Engine/Plugins/Runtime/OpenXR/Source/OpenXRHMD/Private/OpenXRHMD.cpp b/Engine/Plugins/Runtime/OpenXR/Source/OpenXRHMD/Private/OpenXRHMD.cpp
--- a/Engine/Plugins/Runtime/OpenXR/Source/OpenXRHMD/Private/OpenXRHMD.cpp
+++ b/Engine/Plugins/Runtime/OpenXR/Source/OpenXRHMD/Private/OpenXRHMD.cpp
@@ -90,6 +90,12 @@ static TAutoConsoleVariable<int32> CVarOpenXRPreferredViewConfiguration(
 	TEXT("1 = Mono, 2 = Stereo\n"),
 	ECVF_Default);
 
+static TAutoConsoleVariable<bool> CVarOpenXRInvertAlpha(
+	TEXT("xr.OpenXRInvertAlpha"),
+	false,
+	TEXT("Enables alpha inversion of the backgroud layer if the XR_FB_composition_layer_alpha_blend extension is supported.\n"),
+	ECVF_Default);
+
 namespace {
 	static TSet<XrViewConfigurationType> SupportedViewConfigurations{ XR_VIEW_CONFIGURATION_TYPE_PRIMARY_STEREO, XR_VIEW_CONFIGURATION_TYPE_PRIMARY_QUAD_VARJO };
 
@@ -2191,6 +2197,7 @@ bool FOpenXRHMD::StopSession()
 void FOpenXRHMD::OnBeginPlay(FWorldContext& InWorldContext)
 {
 	bOpenXRForceStereoLayersEmulationCVarCachedValue = CVarOpenXRForceStereoLayerEmulation.GetValueOnGameThread();
+	bOpenXRInvertAlphaCvarCachedValue = CVarOpenXRInvertAlpha.GetValueOnGameThread();
 }
 
 void FOpenXRHMD::OnEndPlay(FWorldContext& InWorldContext)
@@ -3313,6 +3320,7 @@ void FOpenXRHMD::OnFinishRendering_RHIThread()
 	{
 		TArray<const XrCompositionLayerBaseHeader*> Headers;
 		XrCompositionLayerProjection Layer = {};
+		XrCompositionLayerAlphaBlendFB LayerAlphaBlend = { XR_TYPE_COMPOSITION_LAYER_ALPHA_BLEND_FB };
 		if (EnumHasAnyFlags(PipelinedLayerStateRHI.LayerStateFlags, EOpenXRLayerStateFlags::SubmitBackgroundLayer))
 		{
 			Layer.type = XR_TYPE_COMPOSITION_LAYER_PROJECTION;
@@ -3323,6 +3331,18 @@ void FOpenXRHMD::OnFinishRendering_RHIThread()
 			Layer.views = PipelinedLayerStateRHI.ProjectionLayers.GetData();
 			Headers.Add(reinterpret_cast<const XrCompositionLayerBaseHeader*>(&Layer));
 
+			if(IsExtensionEnabled(XR_FB_COMPOSITION_LAYER_ALPHA_BLEND_EXTENSION_NAME) &&
+				bOpenXRInvertAlphaCvarCachedValue)
+			{
+				LayerAlphaBlend.next = const_cast<void*>(Layer.next);
+				LayerAlphaBlend.srcFactorColor = PipelinedLayerStateRHI.BasePassLayerBlendParams.srcFactorColor;
+				LayerAlphaBlend.dstFactorColor = PipelinedLayerStateRHI.BasePassLayerBlendParams.dstFactorColor;
+				LayerAlphaBlend.srcFactorAlpha = PipelinedLayerStateRHI.BasePassLayerBlendParams.srcFactorAlpha;
+				LayerAlphaBlend.dstFactorAlpha = PipelinedLayerStateRHI.BasePassLayerBlendParams.dstFactorAlpha;
+
+				Layer.next = &LayerAlphaBlend;
+			}
+
 			for (IOpenXRExtensionPlugin* Module : ExtensionPlugins)
 			{
 				Layer.next = Module->OnEndProjectionLayer(Session, 0, Layer.next, Layer.layerFlags);
diff --git a/Engine/Plugins/Runtime/OpenXR/Source/OpenXRHMD/Private/OpenXRHMD.h b/Engine/Plugins/Runtime/OpenXR/Source/OpenXRHMD/Private/OpenXRHMD.h
--- a/Engine/Plugins/Runtime/OpenXR/Source/OpenXRHMD/Private/OpenXRHMD.h
+++ b/Engine/Plugins/Runtime/OpenXR/Source/OpenXRHMD/Private/OpenXRHMD.h
@@ -106,6 +106,25 @@ public:
 		FXRSwapChainPtr EmulationSwapchain;
 	};
 
+	struct FBasePassLayerBlendParameters
+	{
+		// Default constructor inverts the alpha for color blending to make up for the fact that UE uses
+		// alpha = 0 for opaque and alpha = 1 for transparent while OpenXR does the opposite.
+		// Alpha blending passes through the destination alpha instead.
+		FBasePassLayerBlendParameters()
+		{
+			srcFactorColor = XR_BLEND_FACTOR_ONE_MINUS_SRC_ALPHA_FB;
+			dstFactorColor = XR_BLEND_FACTOR_SRC_ALPHA_FB;
+			srcFactorAlpha = XR_BLEND_FACTOR_ZERO_FB;
+			dstFactorAlpha = XR_BLEND_FACTOR_ONE_FB;
+		}
+		
+		XrBlendFactorFB	srcFactorColor;
+		XrBlendFactorFB	dstFactorColor;
+		XrBlendFactorFB	srcFactorAlpha;
+		XrBlendFactorFB	dstFactorAlpha;
+	};
+
 	enum class EOpenXRLayerStateFlags : uint32
 	{
 		None = 0u,
@@ -132,6 +151,7 @@ public:
 		FEmulatedLayerState EmulatedLayerState;
 
 		EOpenXRLayerStateFlags LayerStateFlags = EOpenXRLayerStateFlags::None;
+		FBasePassLayerBlendParameters BasePassLayerBlendParams;
 	};
 
 	class FVulkanExtensions : public IHeadMountedDisplayVulkanExtensions
@@ -481,6 +501,7 @@ private:
 	FVector					BasePosition;
 
 	bool					bLayerSupportOpenXRCompliant;
+	bool					bOpenXRInvertAlphaCvarCachedValue;
 	bool					bOpenXRForceStereoLayersEmulationCVarCachedValue;
 	TArray<IStereoLayers::FLayerDesc> BackgroundCompositedEmulatedLayers;
 	TArray<IStereoLayers::FLayerDesc> EmulatedFaceLockedLayers;
diff --git a/Engine/Plugins/Runtime/OpenXR/Source/OpenXRHMD/Private/OpenXRHMDModule.cpp b/Engine/Plugins/Runtime/OpenXR/Source/OpenXRHMD/Private/OpenXRHMDModule.cpp
--- a/Engine/Plugins/Runtime/OpenXR/Source/OpenXRHMD/Private/OpenXRHMDModule.cpp
+++ b/Engine/Plugins/Runtime/OpenXR/Source/OpenXRHMD/Private/OpenXRHMDModule.cpp
@@ -504,6 +504,7 @@ bool FOpenXRHMDModule::GetOptionalExtensions(TArray<const ANSICHAR*>& OutExtensi
 	OutExtensions.Add(XR_EXT_DPAD_BINDING_EXTENSION_NAME);
 	OutExtensions.Add(XR_EXT_PALM_POSE_EXTENSION_NAME);
 	OutExtensions.Add(XR_EXT_ACTIVE_ACTION_SET_PRIORITY_EXTENSION_NAME);
+	OutExtensions.Add(XR_FB_COMPOSITION_LAYER_ALPHA_BLEND_EXTENSION_NAME);
 #if PLATFORM_WINDOWS
 	OutExtensions.Add(XR_OCULUS_AUDIO_DEVICE_GUID_EXTENSION_NAME);
 #endif
