OpenXR: properly support throttled spectator screen

AdvanceFrame happens at the end of the frame, where we mistakenly assumed it was the start of the frame.
For platforms that support this feature, we can now advance frame indices after a frame is presented, instead of before the frame begins.

This fixes a bug when stereo would be disabled and enabled.

diff --git a/Engine/Plugins/Runtime/OpenXR/Source/OpenXRHMD/Private/OpenXRHMD_RenderBridge.cpp b/Engine/Plugins/Runtime/OpenXR/Source/OpenXRHMD/Private/OpenXRHMD_RenderBridge.cpp
--- a/Engine/Plugins/Runtime/OpenXR/Source/OpenXRHMD/Private/OpenXRHMD_RenderBridge.cpp
+++ b/Engine/Plugins/Runtime/OpenXR/Source/OpenXRHMD/Private/OpenXRHMD_RenderBridge.cpp
@@ -12,7 +12,7 @@ bool FOpenXRRenderBridge::Present(int32& InOutSyncInterval)
 	if (OpenXRHMD)
 	{
 		OpenXRHMD->OnFinishRendering_RHIThread();
-		bNeedsNativePresent = !OpenXRHMD->IsStandaloneStereoOnlyDevice() && bNativePresent_RHIThread;
+		bNeedsNativePresent = !OpenXRHMD->IsStandaloneStereoOnlyDevice();
 	}
 
 	InOutSyncInterval = 0; // VSync off
diff --git a/Engine/Plugins/Runtime/OpenXR/Source/OpenXRHMD/Private/OpenXRHMD_RenderBridge.h b/Engine/Plugins/Runtime/OpenXR/Source/OpenXRHMD/Private/OpenXRHMD_RenderBridge.h
--- a/Engine/Plugins/Runtime/OpenXR/Source/OpenXRHMD/Private/OpenXRHMD_RenderBridge.h
+++ b/Engine/Plugins/Runtime/OpenXR/Source/OpenXRHMD/Private/OpenXRHMD_RenderBridge.h
@@ -50,30 +50,13 @@ public:
 	/** FRHICustomPresent */
 	virtual bool Present(int32& InOutSyncInterval) override;
 
-	// Called from RenderThread
-	virtual bool NeedsNativePresent() override
-	{
-		return bNativePresent_RenderThread;
-	}
-
 	virtual bool Support10BitSwapchain() const { return false; }
 
 	virtual bool HDRGetMetaDataForStereo(EDisplayOutputFormat& OutDisplayOutputFormat, EDisplayColorGamut& OutDisplayColorGamut, bool& OutbHDRSupported) { return false; }
 
-	virtual void ShouldNativePresent_RenderThread(bool bPresent)
-	{
-		bNativePresent_RenderThread = bPresent;
-	}
-
-	virtual void ShouldNativePresent_RHIThread(bool bPresent)
-	{
-		bNativePresent_RHIThread = bPresent;
-	}
+	virtual void SetSkipRate(uint32 SkipRate) {}
 
 protected:
-	bool bNativePresent_RenderThread = true;
-	bool bNativePresent_RHIThread = true;
-	
 	XrInstance Instance;
 	FOpenXRHMD* OpenXRHMD;
 
