Set the default xr.OpenXRAcquireMode to only acquire on the RHI thread

This is a safer acquire mode that avoids potential threading issues

diff --git a/Engine/Plugins/Runtime/OpenXR/Source/OpenXRHMD/Private/OpenXRHMD.cpp b/Engine/Plugins/Runtime/OpenXR/Source/OpenXRHMD/Private/OpenXRHMD.cpp
--- a/Engine/Plugins/Runtime/OpenXR/Source/OpenXRHMD/Private/OpenXRHMD.cpp
+++ b/Engine/Plugins/Runtime/OpenXR/Source/OpenXRHMD/Private/OpenXRHMD.cpp
@@ -79,7 +79,7 @@ static TAutoConsoleVariable<bool> CVarOpenXRDoNotCopyEmulatedLayersToSpectatorSc
 
 static TAutoConsoleVariable<int32> CVarOpenXRAcquireMode(
 	TEXT("xr.OpenXRAcquireMode"),
-	0,
+	2,
 	TEXT("Override the swapchain acquire mode. 1 = Acquire on any thread, 2 = Only acquire on RHI thread\n"),
 	ECVF_Default);
 
