Quick patch to cover fallback case where all OpenXR view configurations are unsupported.

diff --git a/Engine/Plugins/Runtime/OpenXR/Source/OpenXRHMD/Private/OpenXRHMD.cpp b/Engine/Plugins/Runtime/OpenXR/Source/OpenXRHMD/Private/OpenXRHMD.cpp
--- a/Engine/Plugins/Runtime/OpenXR/Source/OpenXRHMD/Private/OpenXRHMD.cpp
+++ b/Engine/Plugins/Runtime/OpenXR/Source/OpenXRHMD/Private/OpenXRHMD.cpp
@@ -1737,6 +1737,8 @@ bool FOpenXRHMD::OnStereoStartup()
 	// Fill the initial array with valid enum types (this will fail in the validation layer otherwise).
 	ViewConfigTypes.Init(XR_VIEW_CONFIGURATION_TYPE_PRIMARY_MONO, ConfigurationCount);
 	XR_ENSURE(xrEnumerateViewConfigurations(Instance, System, ConfigurationCount, &ConfigurationCount, ViewConfigTypes.GetData()));
+	XrViewConfigurationType PreferredFallbackType = ViewConfigTypes[0];
+	
 	// Filter to supported configurations only
 	ViewConfigTypes = ViewConfigTypes.FilterByPredicate([&](XrViewConfigurationType Type) 
 		{ 
@@ -1764,7 +1766,7 @@ bool FOpenXRHMD::OnStereoStartup()
 	if (!ensure(SelectedViewConfigurationType != XR_VIEW_CONFIGURATION_TYPE_MAX_ENUM))
 	{
 		UE_LOG(LogHMD, Error, TEXT("No compatible view configuration type found, falling back to runtime preferred type."));
-		SelectedViewConfigurationType = ViewConfigTypes[0];
+		SelectedViewConfigurationType = PreferredFallbackType;
 	}
 
 	// Enumerate the views we will be simulating with.
