Remove OpenXR Plugin Views API

diff --git a/Engine/Plugins/Runtime/OpenXR/Source/OpenXRHMD/Private/OpenXRHMD.cpp b/Engine/Plugins/Runtime/OpenXR/Source/OpenXRHMD/Private/OpenXRHMD.cpp
--- a/Engine/Plugins/Runtime/OpenXR/Source/OpenXRHMD/Private/OpenXRHMD.cpp
+++ b/Engine/Plugins/Runtime/OpenXR/Source/OpenXRHMD/Private/OpenXRHMD.cpp
@@ -984,12 +984,6 @@ EStereoscopicPass FOpenXRHMD::GetViewPassForIndex(bool bStereoRequested, int32 V
 	if (!bStereoRequested)
 		return EStereoscopicPass::eSSP_FULL;
 
-	const FPipelinedFrameState& PipelineState = GetPipelinedFrameStateForThread();
-	if (PipelineState.PluginViewInfos.IsValidIndex(ViewIndex) && PipelineState.PluginViewInfos[ViewIndex].bIsPluginManaged)
-	{
-		return PipelineState.PluginViewInfos[ViewIndex].PassType;
-	}
-
 	if (SelectedViewConfigurationType == XR_VIEW_CONFIGURATION_TYPE_PRIMARY_QUAD_VARJO)
 	{
 		return ViewIndex % 2 == 0 ? EStereoscopicPass::eSSP_PRIMARY : EStereoscopicPass::eSSP_SECONDARY;
@@ -1171,12 +1165,6 @@ void FOpenXRHMD::PostRenderViewFamily_RenderThread(FRDGBuilder& GraphBuilder, FS
 			EmulationImage.swapchain = PipelinedLayerStateRendering.EmulatedLayerState.EmulationSwapchain.IsValid() ? static_cast<FOpenXRSwapchain*>(PipelinedLayerStateRendering.EmulatedLayerState.EmulationSwapchain.Get())->GetHandle() : XR_NULL_HANDLE;
 		}
 
-		if (IsViewManagedByPlugin(ViewIndex))
-		{
-			// Plugin owns further usage of the subimage, so we don't use the subimages in our layers
-			continue;
-		}
-
 		XrCompositionLayerProjectionView& Projection = PipelinedLayerStateRendering.ProjectionLayers[ViewIndex];
 
 		Projection.type = XR_TYPE_COMPOSITION_LAYER_PROJECTION_VIEW;
@@ -1476,7 +1464,6 @@ void FOpenXRHMD::EnumerateViews(FPipelinedFrameState& PipelineState)
 	XR_ENSURE(xrEnumerateViewConfigurationViews(Instance, System, SelectedViewConfigurationType, 0, &ViewConfigCount, nullptr));
 	ViewFov.SetNum(ViewConfigCount);
 	PipelineState.ViewConfigs.Empty(ViewConfigCount);
-	PipelineState.PluginViewInfos.Empty(ViewConfigCount);
 	for (uint32 ViewIndex = 0; ViewIndex < ViewConfigCount; ViewIndex++)
 	{
 		XrViewConfigurationView View;
@@ -1491,51 +1478,15 @@ void FOpenXRHMD::EnumerateViews(FPipelinedFrameState& PipelineState)
 			View.next = Module->OnEnumerateViewConfigurationViews(Instance, System, SelectedViewConfigurationType, ViewIndex, View.next);
 		}
 
-		// These are core views that don't have an associated plugin
-		PipelineState.PluginViewInfos.AddDefaulted(1);
 		PipelineState.ViewConfigs.Add(View);
 	}
 	XR_ENSURE(xrEnumerateViewConfigurationViews(Instance, System, SelectedViewConfigurationType, ViewConfigCount, &ViewConfigCount, PipelineState.ViewConfigs.GetData()));
 
-	for (IOpenXRExtensionPlugin* Module : ExtensionPlugins)
-	{
-		TArray<XrViewConfigurationView> ViewConfigs;
-		PRAGMA_DISABLE_DEPRECATION_WARNINGS;
-		Module->GetViewConfigurations(System, ViewConfigs);
-		PRAGMA_ENABLE_DEPRECATION_WARNINGS;
-
-		const EStereoscopicPass PluginPassType = ViewConfigs.Num() > 1 ? EStereoscopicPass::eSSP_PRIMARY : EStereoscopicPass::eSSP_FULL;
-		for (int32 i = 0; i < ViewConfigs.Num(); i++)
-		{
-			PipelineState.PluginViewInfos.Add({ Module, PluginPassType, true });
-		}
-		PipelineState.ViewConfigs.Append(ViewConfigs);
-	}
-	
 	if (Session)
 	{
 		LocateViews(PipelineState, true);
 
 		check(PipelineState.bXrFrameStateUpdated);
-
-		FReadScopeLock DeviceLock(DeviceMutex);
-		for (IOpenXRExtensionPlugin* Module : ExtensionPlugins)
-		{
-			auto Predicate = [Module](const FPluginViewInfo& Info) -> bool
-			{
-				return Info.Plugin == Module;
-			};
-
-			if (PipelineState.PluginViewInfos.ContainsByPredicate(Predicate))
-			{
-				TArray<XrView> Views;
-				PRAGMA_DISABLE_DEPRECATION_WARNINGS;
-				Module->GetViewLocations(Session, PipelineState.FrameState.predictedDisplayTime, DeviceSpaces[HMDDeviceId].Space, Views);
-				PRAGMA_ENABLE_DEPRECATION_WARNINGS;
-				check(Views.Num() > 0);
-				PipelineState.Views.Append(Views);
-			}
-		}
 	}
 	else if (bViewConfigurationFovSupported)
 	{
@@ -1564,18 +1515,6 @@ void FOpenXRHMD::EnumerateViews(FPipelinedFrameState& PipelineState)
 	}
 }
 
-bool FOpenXRHMD::IsViewManagedByPlugin(int32 ViewIndex) const
-{
-	check(IsInRenderingThread());
-	if (!PipelinedFrameStateRendering.PluginViewInfos.IsValidIndex(ViewIndex))
-	{
-		// Was formerly associated with plugin, but plugin is no longer providing this view
-		// Core views always have an entry in the PluginViewInfos array
-		return true;
-	}
-	return PipelinedFrameStateRendering.PluginViewInfos[ViewIndex].bIsPluginManaged;
-}
-
 #if !PLATFORM_HOLOLENS
 void FOpenXRHMD::BuildOcclusionMeshes()
 {
@@ -3387,23 +3326,9 @@ void FOpenXRHMD::OnFinishRendering_RHIThread()
 		EndInfo.layerCount = PipelinedFrameStateRHI.FrameState.shouldRender ? Headers.Num() : 0;
 		EndInfo.layers = PipelinedFrameStateRHI.FrameState.shouldRender ? Headers.GetData() : nullptr;
 
-		// Make callback to plugin including any extra view subimages they've requested
 		for (IOpenXRExtensionPlugin* Module : ExtensionPlugins)
 		{
-			TArray<XrSwapchainSubImage> ColorImages;
-			TArray<XrSwapchainSubImage> DepthImages;
-			for (int32 i = 0; i < PipelinedFrameStateRHI.PluginViewInfos.Num(); i++)
-			{
-				if (PipelinedFrameStateRHI.PluginViewInfos[i].Plugin == Module && PipelinedLayerStateRHI.ColorImages.IsValidIndex(i))
-				{
-					ColorImages.Add(PipelinedLayerStateRHI.ColorImages[i]);
-					if (EnumHasAnyFlags(PipelinedLayerStateRHI.LayerStateFlags, EOpenXRLayerStateFlags::SubmitDepthLayer))
-					{
-						DepthImages.Add(PipelinedLayerStateRHI.DepthImages[i]);
-					}
-				}
-			}
-			EndInfo.next = Module->OnEndFrame(Session, EndInfo.displayTime, ColorImages, DepthImages, EndInfo.next);
+			EndInfo.next = Module->OnEndFrame(Session, EndInfo.displayTime, EndInfo.next);
 		}
 
 		UE_LOG(LogHMD, VeryVerbose, TEXT("Presenting frame predicted to be displayed at %lld"), PipelinedFrameStateRHI.FrameState.predictedDisplayTime);
diff --git a/Engine/Plugins/Runtime/OpenXR/Source/OpenXRHMD/Private/OpenXRHMD.h b/Engine/Plugins/Runtime/OpenXR/Source/OpenXRHMD/Private/OpenXRHMD.h
--- a/Engine/Plugins/Runtime/OpenXR/Source/OpenXRHMD/Private/OpenXRHMD.h
+++ b/Engine/Plugins/Runtime/OpenXR/Source/OpenXRHMD/Private/OpenXRHMD.h
@@ -71,13 +71,6 @@ public:
 		XrPosef BasePose;
 	};
 
-	struct FPluginViewInfo
-	{
-		class IOpenXRExtensionPlugin* Plugin = nullptr;
-		EStereoscopicPass PassType = EStereoscopicPass::eSSP_PRIMARY;
-		bool bIsPluginManaged = false;
-	};
-
 	// The game and render threads each have a separate copy of these structures so that they don't stomp on each other or cause tearing
 	// when the game thread progresses to the next frame while the render thread is still working on the previous frame.
 	struct FPipelinedFrameState
@@ -85,14 +78,11 @@ public:
 		XrFrameState FrameState{XR_TYPE_FRAME_STATE};
 		XrViewState ViewState{XR_TYPE_VIEW_STATE};
 		TArray<XrView> Views;
+		TArray<XrViewConfigurationView> ViewConfigs;
 		TArray<XrSpaceLocation> DeviceLocations;
 		TSharedPtr<FTrackingSpace> TrackingSpace;
 		float WorldToMetersScale = 100.0f;
 		float PixelDensity = 1.0f;
-
-		TArray<XrViewConfigurationView> ViewConfigs;
-		TArray<FPluginViewInfo> PluginViewInfos;
-
 		bool bXrFrameStateUpdated = false;
 	};
 
@@ -283,7 +273,6 @@ protected:
 	void UpdateDeviceLocations(bool bUpdateOpenXRExtensionPlugins);
 	void EnumerateViews(FPipelinedFrameState& PipelineState);
 	void LocateViews(FPipelinedFrameState& PipelinedState, bool ResizeViewsArray = false);
-	bool IsViewManagedByPlugin(int32 ViewIndex) const;
 
 	void CopyTexture_RenderThread(FRHICommandListImmediate& RHICmdList, FRHITexture2D* SrcTexture, FIntRect SrcRect, FRHITexture2D* DstTexture, FIntRect DstRect, 
 								  bool bClearBlack, ERenderTargetActions RTAction, ERHIAccess FinalDstAccess, ETextureCopyBlendModifier SrcTextureCopyModifier) const;
diff --git a/Engine/Plugins/Runtime/OpenXR/Source/OpenXRHMD/Public/IOpenXRExtensionPlugin.h b/Engine/Plugins/Runtime/OpenXR/Source/OpenXRHMD/Public/IOpenXRExtensionPlugin.h
--- a/Engine/Plugins/Runtime/OpenXR/Source/OpenXRHMD/Public/IOpenXRExtensionPlugin.h
+++ b/Engine/Plugins/Runtime/OpenXR/Source/OpenXRHMD/Public/IOpenXRExtensionPlugin.h
@@ -289,22 +289,6 @@ public:
 		return InNext;
 	}
 
-	/**
-	* Callback to provide extra view configurations that should be rendered in the main render pass
-	*/
-	UE_DEPRECATED(5.2, "Support for new view configurations should be added to FOpenXRHMD instead.")
-	virtual void GetViewConfigurations(XrSystemId InSystem, TArray<XrViewConfigurationView>& OutViews)
-	{
-	}
-
-	/**
-	* Callback to provide the pose and fov of each view that was provided in GetViewConfigurations
-	*/
-	UE_DEPRECATED(5.2, "Support for new view configurations should be added to FOpenXRHMD instead.")
-	virtual void GetViewLocations(XrSession InSession, XrTime InDisplayTime, XrSpace InViewSpace, TArray<XrView>& OutViews)
-	{
-	}
-
 	/**
 	* Callbacks with returned pointer added to next chain, do *not* return pointers to structs on the stack.
 	* Remember to assign InNext to the next pointer of your struct or otherwise you may break the next chain.
@@ -393,7 +377,7 @@ public:
 	}
 
 	// FOpenXRRenderBridge::Present, RHI thread
-	virtual const void* OnEndFrame(XrSession InSession, XrTime DisplayTime, const TArray<XrSwapchainSubImage> InColorImages, const TArray<XrSwapchainSubImage> InDepthImages, const void* InNext)
+	virtual const void* OnEndFrame(XrSession InSession, XrTime DisplayTime, const void* InNext)
 	{
 		return InNext;
 	}
