OpenXR: disable depth layer via cvar

diff --git a/Engine/Plugins/Runtime/OpenXR/Source/OpenXRHMD/Private/OpenXRHMD.cpp b/Engine/Plugins/Runtime/OpenXR/Source/OpenXRHMD/Private/OpenXRHMD.cpp
--- a/Engine/Plugins/Runtime/OpenXR/Source/OpenXRHMD/Private/OpenXRHMD.cpp
+++ b/Engine/Plugins/Runtime/OpenXR/Source/OpenXRHMD/Private/OpenXRHMD.cpp
@@ -96,6 +96,12 @@ static TAutoConsoleVariable<bool> CVarOpenXRInvertAlpha(
 	TEXT("Enables alpha inversion of the backgroud layer if the XR_FB_composition_layer_alpha_blend extension is supported.\n"),
 	ECVF_Default);
 
+static TAutoConsoleVariable<bool> CVarOpenXRAllowDepthLayer(
+	TEXT("xr.OpenXRAllowDepthLayer"),
+	true,
+	TEXT("Enables the depth composition layer if the XR_KHR_composition_layer_depth extension is supported.\n"),
+	ECVF_Default);
+
 namespace {
 	static TSet<XrViewConfigurationType> SupportedViewConfigurations{ XR_VIEW_CONFIGURATION_TYPE_PRIMARY_STEREO, XR_VIEW_CONFIGURATION_TYPE_PRIMARY_QUAD_VARJO };
 
@@ -1222,6 +1228,11 @@ bool FOpenXRHMD::IsActiveThisFrame_Internal(const FSceneViewExtensionContext& Co
 
 bool CheckPlatformDepthExtensionSupport(const XrInstanceProperties& InstanceProps)
 {
+	if (!CVarOpenXRAllowDepthLayer.GetValueOnAnyThread())
+	{
+		return false;
+	}
+
 	if (FCStringAnsi::Strstr(InstanceProps.runtimeName, "SteamVR/OpenXR") && RHIGetInterfaceType() == ERHIInterfaceType::Vulkan)
 	{
 		return false;
