Fixed Mobile Preview VRPIE being darker than what was shown in the HMD.

When copying from an sRGB swapchain texture to a rgb10a2_unorm output texture for VRPIE, values were being automatically converted to linear by the texture sample and written as linear to the output texture. However, D3D interprets integer-valued texture for presentation as containing sRGB data so the texture was displayed as darker than intended.

diff --git a/Engine/Plugins/Runtime/OpenXR/Source/OpenXRHMD/Private/OpenXRHMD.cpp b/Engine/Plugins/Runtime/OpenXR/Source/OpenXRHMD/Private/OpenXRHMD.cpp
--- a/Engine/Plugins/Runtime/OpenXR/Source/OpenXRHMD/Private/OpenXRHMD.cpp
+++ b/Engine/Plugins/Runtime/OpenXR/Source/OpenXRHMD/Private/OpenXRHMD.cpp
@@ -3439,20 +3439,14 @@ class FDisplayMappingPS : public FGlobalShader
 public:
 
 	class FArraySource : SHADER_PERMUTATION_BOOL("DISPLAY_MAPPING_PS_FROM_ARRAY");
-	using FPermutationDomain = TShaderPermutationDomain<FArraySource>;
+	class FLinearInput : SHADER_PERMUTATION_BOOL("DISPLAY_MAPPING_INPUT_IS_LINEAR");
+	using FPermutationDomain = TShaderPermutationDomain<FArraySource, FLinearInput>;
 
 	static bool ShouldCompilePermutation(const FGlobalShaderPermutationParameters& Parameters)
 	{
-		return IsFeatureLevelSupported(Parameters.Platform, ERHIFeatureLevel::SM5);
-	}
-
-	static void ModifyCompilationEnvironment(const FGlobalShaderPermutationParameters& Parameters, FShaderCompilerEnvironment& OutEnvironment)
-	{
-		FGlobalShader::ModifyCompilationEnvironment(Parameters, OutEnvironment);
-		OutEnvironment.SetDefine(TEXT("DISPLAY_MAPPING_PS"), 1);
+		return true;
 	}
 
-
 	FDisplayMappingPS(const ShaderMetaType::CompiledShaderInitializerType& Initializer) :
 		FGlobalShader(Initializer)
 	{
@@ -3560,7 +3554,7 @@ public:
 
 	static const TCHAR* GetSourceFilename()
 	{
-		return TEXT("/Engine/Private/CompositeUIPixelShader.usf");
+		return TEXT("/Engine/Private/DisplayMappingPixelShader.usf");
 	}
 
 	static const TCHAR* GetFunctionName()
@@ -3576,7 +3570,7 @@ private:
 	LAYOUT_FIELD(FShaderResourceParameter, SceneSampler);
 };
 
-IMPLEMENT_SHADER_TYPE(, FDisplayMappingPS, TEXT("/Engine/Private/CompositeUIPixelShader.usf"), TEXT("DisplayMappingPS"), SF_Pixel);
+IMPLEMENT_SHADER_TYPE(, FDisplayMappingPS, TEXT("/Engine/Private/DisplayMappingPixelShader.usf"), TEXT("DisplayMappingPS"), SF_Pixel);
 
 // We use CopyTexture for a number of subtley different use cases.
 // * SpectatorScreen, background layer - optional clear to black, opaque
@@ -3678,6 +3672,7 @@ void FOpenXRHMD::CopyTexture_RenderThread(FRHICommandListImmediate& RHICmdList,
 		TShaderRef<FScreenPS> ScreenPS;
 
 		bool bNeedsDisplayMapping = false;
+		bool bIsInputLinear = false;
 		EDisplayOutputFormat TVDisplayOutputFormat = EDisplayOutputFormat::SDR_sRGB;
 		EDisplayColorGamut HMDColorGamut = EDisplayColorGamut::sRGB_D65;
 		EDisplayColorGamut TVColorGamut = EDisplayColorGamut::sRGB_D65;
@@ -3697,6 +3692,15 @@ void FOpenXRHMD::CopyTexture_RenderThread(FRHICommandListImmediate& RHICmdList,
 					bNeedsDisplayMapping = true;
 				}
 			}
+
+			// In Android Vulkan preview, when the sRGB swapchain texture is sampled, the data is converted to linear and written to the RGBA10A2_UNORM texture.
+			// However, D3D interprets integer-valued display formats as containing sRGB data, so we need to convert the linear data back to sRGB.
+			if (!IsMobileHDR() && IsMobilePlatform(GetConfiguredShaderPlatform()) && IsSimulatedPlatform(GetConfiguredShaderPlatform()))
+			{
+				bNeedsDisplayMapping = true;
+				TVDisplayOutputFormat = EDisplayOutputFormat::SDR_sRGB;
+				bIsInputLinear = true;
+			}
 		}
 
 		bNeedsDisplayMapping &= IsFeatureLevelSupported(GetConfiguredShaderPlatform(), ERHIFeatureLevel::ES3_1);
@@ -3707,11 +3711,10 @@ void FOpenXRHMD::CopyTexture_RenderThread(FRHICommandListImmediate& RHICmdList,
 		{
 			FDisplayMappingPS::FPermutationDomain PermutationVector;
 			PermutationVector.Set<FDisplayMappingPS::FArraySource>(bIsArraySource);
+			PermutationVector.Set<FDisplayMappingPS::FLinearInput>(bIsInputLinear);
 
 			TShaderMapRef<FDisplayMappingPS> DisplayMappingPSRef(ShaderMap, PermutationVector);
 
-			ensureMsgf(!bIsArraySource, TEXT("Stub implementation for array source exists, but this path is not fully supported because post-processing not working with MMV"));
-
 			DisplayMappingPS = DisplayMappingPSRef;
 			PixelShader = DisplayMappingPSRef;
 		}
diff --git a/Engine/Shaders/Private/CompositeUIPixelShader.usf b/Engine/Shaders/Private/CompositeUIPixelShader.usf
--- a/Engine/Shaders/Private/CompositeUIPixelShader.usf
+++ b/Engine/Shaders/Private/CompositeUIPixelShader.usf
@@ -84,25 +84,10 @@ void GenerateLUTPS(
 Texture2D<uint> UIWriteMaskTexture;
 #endif
 
-#ifndef DISPLAY_MAPPING_PS_FROM_ARRAY
-	#define DISPLAY_MAPPING_PS_FROM_ARRAY 0
-#endif // DISPLAY_MAPPING_PS_FROM_ARRAY
-
-#if DISPLAY_MAPPING_PS_FROM_ARRAY
-	// if source is an array, we only need the first slice
-	#define TextureOrArray2D						Texture2DArray
-	#define TexOrArr2DSample(Tex, S, UV)			Tex.Sample(S, float3(UV.x, UV.y, 0))
-	#define TexOrArr2DSampleLevel(Tex, S, UV, L)	Tex.SampleLevel(S, float3(UV.x, UV.y, 0), L)
-#else
-	#define TextureOrArray2D						Texture2D
-	#define TexOrArr2DSample(Tex, S, UV)			Texture2DSample(Tex, S, UV)
-	#define TexOrArr2DSampleLevel(Tex, S, UV, L)	Texture2DSampleLevel(Tex, S, UV, L)
-#endif
-
 Texture2D UITexture;
 SamplerState UISampler;
 
-TextureOrArray2D SceneTexture;
+Texture2D SceneTexture;
 SamplerState SceneSampler;
 
 Texture3D ColorSpaceLUT;
@@ -142,7 +127,7 @@ void HDRBackBufferConvertPS(
 	)
 {
 	OutColor = 0;
-	float3 SceneColor = TexOrArr2DSample(SceneTexture, SceneSampler, Input.UV).xyz ;
+	float3 SceneColor = Texture2DSample(SceneTexture, SceneSampler, Input.UV).xyz ;
 	SceneColor = ST2084ToLinear( SceneColor ) / ScRGBScaleFactor; 
 	OutColor.rgb = SceneColor;
 }
@@ -155,7 +140,7 @@ void HDREditorConvert(
 	)
 {
 	OutColor = 0;
-	float3 SceneColor = TexOrArr2DSample(SceneTexture, SceneSampler, Input.UV).xyz ;
+	float3 SceneColor = Texture2DSample(SceneTexture, SceneSampler, Input.UV).xyz ;
 
 	// Get color into Rec2020 space
 	SceneColor = sRGBToLinear( SceneColor) ;
@@ -267,7 +252,7 @@ void Main(
 	}
 #endif
 
-	float3 SceneColor = TexOrArr2DSample(SceneTexture, SceneSampler, Input.UV).xyz;
+	float3 SceneColor = Texture2DSample(SceneTexture, SceneSampler, Input.UV).xyz;
 	SceneColor = ComputeHDRSceneColor(SceneColor);
 	OutColor.xyz = ComposeUIAndScene(SceneColor.xyz, UIColor, UILevel);
 
@@ -342,33 +327,6 @@ void CompositeUICS(uint3 DispatchThreadId : SV_DispatchThreadID)
 
 #endif
 
-#if DISPLAY_MAPPING_PS
-float4x4 TextureToOutputGamutMatrix;
-
-// Currently hard-coded to go from G2.2 to either PQ/Rec2020 or sRGB
-void DisplayMappingPS(
-	FScreenVertexOutput Input,
-	out float4 OutColor : SV_Target0
-	)
-{
-	const float4 SceneColor = TexOrArr2DSampleLevel(SceneTexture, SceneSampler, Input.UV, 0);
-	OutColor = SceneColor;
-
-	float3 LinearSceneColor = pow(SceneColor.rgb, 2.2);
-	LinearSceneColor = mul((float3x3)TextureToOutputGamutMatrix, LinearSceneColor);
-	
-	if(OutputDevice == TONEMAPPER_OUTPUT_ACES1000nitST2084 || OutputDevice == TONEMAPPER_OUTPUT_ACES2000nitST2084)
-	{
-		OutColor.xyz = LinearToST2084(LinearSceneColor*300.f);
-	}
-	else
-	{
-		OutColor.xyz = LinearToSrgbBranchless(saturate(LinearSceneColor));
-	}
-}
-
-#endif
-
 #if BLIT_UI_TO_HDR
 
 float2 UITextureSize;
@@ -389,7 +347,7 @@ void BlitUIToHDRPS(
 	}
 #endif
 
-	float3 SceneColor = TexOrArr2DSample(SceneTexture, UISampler, Input.UV).xyz;
+	float3 SceneColor = Texture2DSample(SceneTexture, UISampler, Input.UV).xyz;
 	SceneColor = ComputeHDRSceneColor(SceneColor);
 	OutColor.xyz = ComposeUIAndScene(SceneColor.xyz, UIColor, UILevel) / ScRGBScaleFactor;
 	OutColor.a = 1;
diff --git a/Engine/Shaders/Private/DisplayMappingPixelShader.usf b/Engine/Shaders/Private/DisplayMappingPixelShader.usf
new file mode 100644
--- /dev/null
+++ b/Engine/Shaders/Private/DisplayMappingPixelShader.usf
@@ -0,0 +1,65 @@
+
+/*=============================================================================
+DisplayMapping.usf: map input texture to output texture color space.
+=============================================================================*/
+
+#include "Common.ush"
+#include "GammaCorrectionCommon.ush"
+#include "TonemapCommon.ush"
+
+#ifndef DISPLAY_MAPPING_INPUT_IS_LINEAR
+	#define DISPLAY_MAPPING_INPUT_IS_LINEAR 0
+#endif // DISPLAY_MAPPING_INPUT_IS_LINEAR
+
+#ifndef DISPLAY_MAPPING_PS_FROM_ARRAY
+	#define DISPLAY_MAPPING_PS_FROM_ARRAY 0
+#endif // DISPLAY_MAPPING_PS_FROM_ARRAY
+
+#if DISPLAY_MAPPING_PS_FROM_ARRAY
+	// if source is an array, we only need the first slice
+	#define TextureOrArray2D						Texture2DArray
+	#define TexOrArr2DSample(Tex, S, UV)			Tex.Sample(S, float3(UV.x, UV.y, 0))
+	#define TexOrArr2DSampleLevel(Tex, S, UV, L)	Tex.SampleLevel(S, float3(UV.x, UV.y, 0), L)
+#else
+	#define TextureOrArray2D						Texture2D
+	#define TexOrArr2DSample(Tex, S, UV)			Texture2DSample(Tex, S, UV)
+	#define TexOrArr2DSampleLevel(Tex, S, UV, L)	Texture2DSampleLevel(Tex, S, UV, L)
+#endif
+
+uint OutputDevice;
+uint OutputGamut;
+float OutputMaxLuminance;
+
+TextureOrArray2D SceneTexture;
+SamplerState SceneSampler;
+float4x4 TextureToOutputGamutMatrix;
+
+float3 GetLinearColor(float3 InColor)
+{
+#if DISPLAY_MAPPING_INPUT_IS_LINEAR
+    return InColor;
+#else
+    return sRGBToLinear(InColor);
+#endif
+}
+
+void DisplayMappingPS(
+	FScreenVertexOutput Input,
+	out float4 OutColor : SV_Target0
+	)
+{
+	const float4 SceneColor = TexOrArr2DSampleLevel(SceneTexture, SceneSampler, Input.UV, 0);
+	OutColor = SceneColor;
+
+	float3 LinearSceneColor = GetLinearColor(SceneColor.rgb);
+	LinearSceneColor = mul((float3x3)TextureToOutputGamutMatrix, LinearSceneColor);
+	
+	if(OutputDevice == TONEMAPPER_OUTPUT_ACES1000nitST2084 || OutputDevice == TONEMAPPER_OUTPUT_ACES2000nitST2084)
+	{
+		OutColor.xyz = LinearToST2084(LinearSceneColor*300.f);
+	}
+	else
+	{
+		OutColor.xyz = LinearToSrgbBranchless(saturate(LinearSceneColor));
+	}
+}
\ No newline at end of file
