OpenXR: log the extensions that the runtime supports

diff --git a/Engine/Plugins/Runtime/OpenXR/Source/OpenXRHMD/Private/OpenXRHMDModule.cpp b/Engine/Plugins/Runtime/OpenXR/Source/OpenXRHMD/Private/OpenXRHMDModule.cpp
--- a/Engine/Plugins/Runtime/OpenXR/Source/OpenXRHMD/Private/OpenXRHMDModule.cpp
+++ b/Engine/Plugins/Runtime/OpenXR/Source/OpenXRHMD/Private/OpenXRHMDModule.cpp
@@ -203,9 +203,11 @@ bool FOpenXRHMDModule::EnumerateExtensions()
 
 	if (XR_ENSURE(xrEnumerateInstanceExtensionProperties(nullptr, ExtensionsCount, &ExtensionsCount, Properties.GetData())))
 	{
+		UE_LOG(LogHMD, Log, TEXT("OpenXR runtime supported extensions:"));
 		for (const XrExtensionProperties& Prop : Properties)
 		{
 			AvailableExtensions.Add(Prop.extensionName);
+			UE_LOG(LogHMD, Log, TEXT("\t%S"), (Prop.extensionName));
 		}
 		return true;
 	}
