Add support for XR_FB_foveation OpenXR extension.

Added a foveation swapchain in OpenXRHMD and implemented IVariableRateShadingImageGenerator to update the foveation swapchain to the latest settings and provide the foveation images to RenderGraph. Also added the possibility in OpenXRHMD_Swapchain to create a foveation swapchain instead of a normal swapchain. This only works when using either Vulkan or OpenGL ES.

diff --git a/Engine/Plugins/Runtime/OpenXR/Source/OpenXRHMD/Private/OpenXRHMD.cpp b/Engine/Plugins/Runtime/OpenXR/Source/OpenXRHMD/Private/OpenXRHMD.cpp
--- a/Engine/Plugins/Runtime/OpenXR/Source/OpenXRHMD/Private/OpenXRHMD.cpp
+++ b/Engine/Plugins/Runtime/OpenXR/Source/OpenXRHMD/Private/OpenXRHMD.cpp
@@ -102,6 +102,32 @@ static TAutoConsoleVariable<bool> CVarOpenXRAllowDepthLayer(
 	TEXT("Enables the depth composition layer if the XR_KHR_composition_layer_depth extension is supported.\n"),
 	ECVF_Default);
 
+static TAutoConsoleVariable<int32> CVarOpenXRFBFoveationLevel(
+	TEXT("xr.OpenXRFBFoveationLevel"),
+	0,
+	TEXT("Possible foveation levels as specified by the XrFoveationLevelFB enumeration.\n")
+	TEXT("0 = None, 1 = Low , 2 = Medium, 3 = High.\n"),
+	ECVF_Default);
+
+static TAutoConsoleVariable<bool> CVarOpenXRFBFoveationDynamic(
+	TEXT("xr.OpenXRFBFoveationDynamic"),
+	false,
+	TEXT("Whether dynamic changing foveation based on performance headroom is enabled.\n"),
+	ECVF_Default);
+
+static TAutoConsoleVariable<float> CVarOpenXRFBFoveationVerticalOffset(
+	TEXT("xr.OpenXRFBFoveationVerticalOffset"),
+	0,
+	TEXT("Desired vertical offset in degrees for the center of the foveation pattern.\n"),
+	ECVF_Default);
+
+static TAutoConsoleVariable<int32> CVarFBFoveationPreview(
+	TEXT("xr.OpenXRFBFoveation.Preview"),
+	1,
+	TEXT("Whether to include FB foveation in VRS preview overlay. This is currently not implemented.")
+	TEXT("0 - off, 1 - on (default)"),
+	ECVF_RenderThreadSafe);
+
 namespace {
 	static TSet<XrViewConfigurationType> SupportedViewConfigurations{ XR_VIEW_CONFIGURATION_TYPE_PRIMARY_STEREO, XR_VIEW_CONFIGURATION_TYPE_PRIMARY_QUAD_VARJO };
 
@@ -1311,6 +1337,30 @@ FOpenXRHMD::FOpenXRHMD(const FAutoRegister& AutoRegister, XrInstance InInstance,
 	bIsAcquireOnAnyThreadSupported = CheckPlatformAcquireOnAnyThreadSupport(InstanceProperties);
 	ReconfigureForShaderPlatform(GMaxRHIShaderPlatform);
 
+	bFoveationExtensionSupported = IsExtensionEnabled(XR_FB_SWAPCHAIN_UPDATE_STATE_EXTENSION_NAME) &&
+		IsExtensionEnabled(XR_FB_FOVEATION_EXTENSION_NAME) &&
+		IsExtensionEnabled(XR_FB_FOVEATION_CONFIGURATION_EXTENSION_NAME);
+
+#ifdef XR_USE_GRAPHICS_API_VULKAN
+	bFoveationExtensionSupported &= IsExtensionEnabled(XR_FB_FOVEATION_VULKAN_EXTENSION_NAME) && GRHISupportsAttachmentVariableRateShading;
+#endif
+
+	if(bFoveationExtensionSupported)
+	{
+		XR_ENSURE(xrGetInstanceProcAddr(Instance, "xrCreateFoveationProfileFB", (PFN_xrVoidFunction*)&xrCreateFoveationProfileFB));
+		XR_ENSURE(xrGetInstanceProcAddr(Instance, "xrUpdateSwapchainFB", (PFN_xrVoidFunction*)&xrUpdateSwapchainFB));
+		XR_ENSURE(xrGetInstanceProcAddr(Instance, "xrDestroyFoveationProfileFB", (PFN_xrVoidFunction*)&xrDestroyFoveationProfileFB));
+
+		int32 SanitisedFoveationLevel = FMath::Clamp(CVarOpenXRFBFoveationLevel->GetInt(), 0, 3);
+		bool bFoveationDynamic = CVarOpenXRFBFoveationDynamic->GetBool();
+		float SanitisedVerticalOffset = CVarOpenXRFBFoveationVerticalOffset->GetFloat();
+		SanitisedVerticalOffset = SanitisedVerticalOffset >= 0 ? SanitisedVerticalOffset : 0;
+
+		FoveationLevel = static_cast<XrFoveationLevelFB>(SanitisedFoveationLevel);
+		VerticalOffset = SanitisedVerticalOffset;
+		FoveationDynamic = bFoveationDynamic ? XR_FOVEATION_DYNAMIC_LEVEL_ENABLED_FB : XR_FOVEATION_DYNAMIC_DISABLED_FB;
+	}
+
 #if PLATFORM_HOLOLENS || PLATFORM_ANDROID
 	bIsStandaloneStereoOnlyDevice = IStereoRendering::IsStartInVR();
 #else
@@ -2211,6 +2261,12 @@ bool FOpenXRHMD::AllocateRenderTargetTextures(uint32 SizeX, uint32 SizeY, uint8
 	{
 		UnifiedCreateFlags |= TexCreate_SRGB;
 	}
+	ETextureCreateFlags AuxiliaryCreateFlags = ETextureCreateFlags::None;
+
+	if(bFoveationExtensionSupported)
+	{
+		AuxiliaryCreateFlags |= TexCreate_Foveation;
+	}
 
 	// Temporary workaround to swapchain formats - OpenXR doesn't support 10-bit sRGB swapchains, so prefer 8-bit sRGB instead.
 	if (Format == PF_A2B10G10R10 && !RenderBridge->Support10BitSwapchain())
@@ -2229,7 +2285,7 @@ bool FOpenXRHMD::AllocateRenderTargetTextures(uint32 SizeX, uint32 SizeY, uint8
 	{
 		ensureMsgf(NumSamples == 1, TEXT("OpenXR supports MSAA swapchains, but engine logic expects the swapchain target to be 1x."));
 
-		Swapchain = RenderBridge->CreateSwapchain(Session, Format, ActualFormat, SizeX, SizeY, bIsMobileMultiViewEnabled ? 2 : 1, NumMips, NumSamples, UnifiedCreateFlags, ClearColor);
+		Swapchain = RenderBridge->CreateSwapchain(Session, Format, ActualFormat, SizeX, SizeY, bIsMobileMultiViewEnabled ? 2 : 1, NumMips, NumSamples, UnifiedCreateFlags, ClearColor, AuxiliaryCreateFlags);
 		if (!Swapchain)
 		{
 			return false;
@@ -2242,6 +2298,12 @@ bool FOpenXRHMD::AllocateRenderTargetTextures(uint32 SizeX, uint32 SizeY, uint8
 				Swapchain->IncrementSwapChainIndex_RHIThread();
 			});
 		}
+		if (bFoveationExtensionSupported)
+		{
+			UpdateFoveationImages();
+			FOpenXRSwapchain* ColorSwapchain = static_cast<FOpenXRSwapchain*>(Swapchain.Get());
+			ColorSwapchain->GetFragmentDensityMaps(FoveationImages, bIsMobileMultiViewEnabled);
+		}
 	}
 
 	// Grab the presentation texture out of the swapchain.
@@ -3885,6 +3947,74 @@ void FOpenXRHMD::UpdateLayer(FOpenXRLayer& ManagerLayer, uint32 LayerId, bool bI
 	});
 }
 
+//---------------------------------------------------
+// IVariableRateShadingImageGenerator interface for XR_FB_foveation Implementation
+//---------------------------------------------------
+
+FRDGTextureRef FOpenXRHMD::GetImage(FRDGBuilder& GraphBuilder, const FViewInfo& ViewInfo, FVariableRateShadingImageManager::EVRSImageType ImageType)
+{
+	if (!bFoveationExtensionSupported)
+	{
+		return nullptr;
+	}
+	if (PipelinedLayerStateRendering.ColorSwapchain != nullptr && !FoveationImages.IsEmpty())
+	{
+		FOpenXRSwapchain* ColorSwapchain = static_cast<FOpenXRSwapchain*>(PipelinedLayerStateRendering.ColorSwapchain.Get());
+		
+		FTextureRHIRef SwapchainTexture = FoveationImages[ColorSwapchain->GetSwapChainIndex_RHIThread()];
+		TRefCountPtr<IPooledRenderTarget> PooledRenderTarget = CreateRenderTarget(SwapchainTexture, *SwapchainTexture->GetName().ToString());
+
+		return GraphBuilder.RegisterExternalTexture(PooledRenderTarget, *SwapchainTexture->GetName().ToString(), ERDGTextureFlags::SkipTracking);
+	}
+	else 
+	{
+		UE_LOG(LogHMD, Error, TEXT("No valid color swapchain to get foveation images from."));
+		return nullptr;
+	}
+}
+
+void FOpenXRHMD::PrepareImages(FRDGBuilder& GraphBuilder, const FSceneViewFamily& ViewFamily, const FMinimalSceneTextures& SceneTextures)
+{
+	return; //Currently not implemented as the images are prepared only when the color swapchain is reallocated.
+}
+
+bool FOpenXRHMD::IsEnabledForView(const FSceneView& View) const
+{
+	return IsStereoEnabled() && bFoveationExtensionSupported && View.StereoPass != EStereoscopicPass::eSSP_FULL;
+}
+
+// This is currently not implemented.
+FRDGTextureRef FOpenXRHMD::GetDebugImage(FRDGBuilder& GraphBuilder, const FViewInfo& ViewInfo, FVariableRateShadingImageManager::EVRSImageType ImageType)
+{
+	return nullptr;
+}
+
+void FOpenXRHMD::UpdateFoveationImages()
+{
+	XrFoveationLevelProfileCreateInfoFB FoveationLevelProfileInfo{ XR_TYPE_FOVEATION_LEVEL_PROFILE_CREATE_INFO_FB };
+	FoveationLevelProfileInfo.next = nullptr;
+	FoveationLevelProfileInfo.level = FoveationLevel;
+	FoveationLevelProfileInfo.verticalOffset = VerticalOffset;
+	FoveationLevelProfileInfo.dynamic = FoveationDynamic;
+
+	XrFoveationProfileCreateInfoFB FoveationCreateInfo{ XR_TYPE_FOVEATION_PROFILE_CREATE_INFO_FB };
+	FoveationCreateInfo.next = &FoveationLevelProfileInfo;
+
+	XrFoveationProfileFB FoveationProfile = XR_NULL_HANDLE;
+	XR_ENSURE(xrCreateFoveationProfileFB(Session, &FoveationCreateInfo, &FoveationProfile));
+
+	XrSwapchainStateFoveationFB SwapchainFoveationState{ XR_TYPE_SWAPCHAIN_STATE_FOVEATION_FB };
+	SwapchainFoveationState.flags = 0; // As per OpenXR specification.
+	SwapchainFoveationState.next = nullptr;
+	SwapchainFoveationState.profile = FoveationProfile;
+
+	FOpenXRSwapchain* ColorSwapchain = static_cast<FOpenXRSwapchain*>(PipelinedLayerStateRendering.ColorSwapchain.Get());
+
+	XR_ENSURE(xrUpdateSwapchainFB(ColorSwapchain->GetHandle(), reinterpret_cast<const XrSwapchainStateBaseHeaderFB*>(&SwapchainFoveationState)));
+
+	XR_ENSURE(xrDestroyFoveationProfileFB(FoveationProfile));
+}
+
 //---------------------------------------------------
 // OpenXR Action Space Implementation
 //---------------------------------------------------
diff --git a/Engine/Plugins/Runtime/OpenXR/Source/OpenXRHMD/Private/OpenXRHMD.h b/Engine/Plugins/Runtime/OpenXR/Source/OpenXRHMD/Private/OpenXRHMD.h
--- a/Engine/Plugins/Runtime/OpenXR/Source/OpenXRHMD/Private/OpenXRHMD.h
+++ b/Engine/Plugins/Runtime/OpenXR/Source/OpenXRHMD/Private/OpenXRHMD.h
@@ -16,6 +16,7 @@
 #include "IHeadMountedDisplayVulkanExtensions.h"
 #include "IOpenXRExtensionPluginDelegates.h"
 #include "Misc/EnumClassFlags.h"
+#include "VariableRateShadingImageManager.h"
 
 #include <openxr/openxr.h>
 
@@ -38,6 +39,7 @@ class FOpenXRHMD
 	, public FOpenXRAssetManager
 	, public TStereoLayerManager<FOpenXRLayer>
 	, public IOpenXRExtensionPluginDelegates
+	, public IVariableRateShadingImageGenerator
 {
 private:
 
@@ -369,6 +371,19 @@ public:
 	/** IStereoLayers */
 	virtual bool ShouldCopyDebugLayersToSpectatorScreen() const override { return true; }
 
+	/** IVariableRateShadingImageGenerator interface for XR_FB_foveation */
+	virtual FRDGTextureRef GetImage(FRDGBuilder& GraphBuilder, const FViewInfo& ViewInfo, FVariableRateShadingImageManager::EVRSImageType ImageType) override;
+	virtual void PrepareImages(FRDGBuilder& GraphBuilder, const FSceneViewFamily& ViewFamily, const FMinimalSceneTextures& SceneTextures) override;
+	virtual bool IsEnabledForView(const FSceneView& View) const override;
+	virtual FRDGTextureRef GetDebugImage(FRDGBuilder& GraphBuilder, const FViewInfo& ViewInfo, FVariableRateShadingImageManager::EVRSImageType ImageType) override;
+	virtual FVariableRateShadingImageManager::EVRSSourceType GetType() const override
+	{
+		return FVariableRateShadingImageManager::EVRSSourceType::FixedFoveation;
+	}
+
+	void UpdateFoveationImages();
+
+
 	/** IOpenXRExtensionPluginDelegates */
 public:
 	virtual FApplyHapticFeedbackAddChainStructsDelegate& GetApplyHapticFeedbackAddChainStructsDelegate() override { return ApplyHapticFeedbackAddChainStructsDelegate; }
@@ -495,6 +510,16 @@ private:
 	TArray<IStereoLayers::FLayerDesc> BackgroundCompositedEmulatedLayers;
 	TArray<IStereoLayers::FLayerDesc> EmulatedFaceLockedLayers;
 	TArray<FOpenXRLayer>			  NativeLayers;
+
+	// XR_FB_foveation
+	bool					bFoveationExtensionSupported;
+	XrFoveationLevelFB		FoveationLevel;
+	float					VerticalOffset;
+	XrFoveationDynamicFB	FoveationDynamic;
+	TArray<FTextureRHIRef>	FoveationImages;
+	PFN_xrCreateFoveationProfileFB	xrCreateFoveationProfileFB;
+	PFN_xrUpdateSwapchainFB			xrUpdateSwapchainFB;
+	PFN_xrDestroyFoveationProfileFB xrDestroyFoveationProfileFB;
 };
 
 ENUM_CLASS_FLAGS(FOpenXRHMD::EOpenXRLayerStateFlags);
\ No newline at end of file
diff --git a/Engine/Plugins/Runtime/OpenXR/Source/OpenXRHMD/Private/OpenXRHMDModule.cpp b/Engine/Plugins/Runtime/OpenXR/Source/OpenXRHMD/Private/OpenXRHMDModule.cpp
--- a/Engine/Plugins/Runtime/OpenXR/Source/OpenXRHMD/Private/OpenXRHMDModule.cpp
+++ b/Engine/Plugins/Runtime/OpenXR/Source/OpenXRHMD/Private/OpenXRHMDModule.cpp
@@ -495,6 +495,7 @@ bool FOpenXRHMDModule::GetOptionalExtensions(TArray<const ANSICHAR*>& OutExtensi
 #ifdef XR_USE_GRAPHICS_API_VULKAN
 	OutExtensions.Add(XR_KHR_VULKAN_ENABLE_EXTENSION_NAME);
 	OutExtensions.Add(XR_KHR_VULKAN_SWAPCHAIN_FORMAT_LIST_EXTENSION_NAME);
+	OutExtensions.Add(XR_FB_FOVEATION_VULKAN_EXTENSION_NAME);
 #endif
 	OutExtensions.Add(XR_KHR_COMPOSITION_LAYER_DEPTH_EXTENSION_NAME);
 	OutExtensions.Add(XR_KHR_COMPOSITION_LAYER_CYLINDER_EXTENSION_NAME);
@@ -507,6 +508,9 @@ bool FOpenXRHMDModule::GetOptionalExtensions(TArray<const ANSICHAR*>& OutExtensi
 	OutExtensions.Add(XR_EXT_PALM_POSE_EXTENSION_NAME);
 	OutExtensions.Add(XR_EXT_ACTIVE_ACTION_SET_PRIORITY_EXTENSION_NAME);
 	OutExtensions.Add(XR_FB_COMPOSITION_LAYER_ALPHA_BLEND_EXTENSION_NAME);
+	OutExtensions.Add(XR_FB_FOVEATION_EXTENSION_NAME);
+	OutExtensions.Add(XR_FB_SWAPCHAIN_UPDATE_STATE_EXTENSION_NAME);
+	OutExtensions.Add(XR_FB_FOVEATION_CONFIGURATION_EXTENSION_NAME); 
 #if PLATFORM_WINDOWS
 	OutExtensions.Add(XR_OCULUS_AUDIO_DEVICE_GUID_EXTENSION_NAME);
 #endif
diff --git a/Engine/Plugins/Runtime/OpenXR/Source/OpenXRHMD/Private/OpenXRHMD_RenderBridge.cpp b/Engine/Plugins/Runtime/OpenXR/Source/OpenXRHMD/Private/OpenXRHMD_RenderBridge.cpp
--- a/Engine/Plugins/Runtime/OpenXR/Source/OpenXRHMD/Private/OpenXRHMD_RenderBridge.cpp
+++ b/Engine/Plugins/Runtime/OpenXR/Source/OpenXRHMD/Private/OpenXRHMD_RenderBridge.cpp
@@ -58,9 +58,9 @@ public:
 		return &Binding;
 	}
 
-	virtual FXRSwapChainPtr CreateSwapchain(XrSession InSession, uint8 Format, uint8& OutActualFormat, uint32 SizeX, uint32 SizeY, uint32 ArraySize, uint32 NumMips, uint32 NumSamples, ETextureCreateFlags CreateFlags, const FClearValueBinding& ClearValueBinding) override final
+	virtual FXRSwapChainPtr CreateSwapchain(XrSession InSession, uint8 Format, uint8& OutActualFormat, uint32 SizeX, uint32 SizeY, uint32 ArraySize, uint32 NumMips, uint32 NumSamples, ETextureCreateFlags CreateFlags, const FClearValueBinding& ClearValueBinding, ETextureCreateFlags AuxiliaryCreateFlags) override final
 	{
-		return CreateSwapchain_D3D11(InSession, Format, OutActualFormat, SizeX, SizeY, ArraySize, NumMips, NumSamples, CreateFlags, ClearValueBinding);
+		return CreateSwapchain_D3D11(InSession, Format, OutActualFormat, SizeX, SizeY, ArraySize, NumMips, NumSamples, CreateFlags, ClearValueBinding, AuxiliaryCreateFlags);
 	}
 
 private:
@@ -109,9 +109,9 @@ public:
 		return &Binding;
 	}
 
-	virtual FXRSwapChainPtr CreateSwapchain(XrSession InSession, uint8 Format, uint8& OutActualFormat, uint32 SizeX, uint32 SizeY, uint32 ArraySize, uint32 NumMips, uint32 NumSamples, ETextureCreateFlags CreateFlags, const FClearValueBinding& ClearValueBinding) override final
+	virtual FXRSwapChainPtr CreateSwapchain(XrSession InSession, uint8 Format, uint8& OutActualFormat, uint32 SizeX, uint32 SizeY, uint32 ArraySize, uint32 NumMips, uint32 NumSamples, ETextureCreateFlags CreateFlags, const FClearValueBinding& ClearValueBinding, ETextureCreateFlags AuxiliaryCreateFlags) override final
 	{
-		return CreateSwapchain_D3D12(InSession, Format, OutActualFormat, SizeX, SizeY, ArraySize, NumMips, NumSamples, CreateFlags, ClearValueBinding);
+		return CreateSwapchain_D3D12(InSession, Format, OutActualFormat, SizeX, SizeY, ArraySize, NumMips, NumSamples, CreateFlags, ClearValueBinding, AuxiliaryCreateFlags);
 	}
 
 private:
@@ -165,9 +165,9 @@ public:
 		return nullptr;
 	}
 
-	virtual FXRSwapChainPtr CreateSwapchain(XrSession InSession, uint8 Format, uint8& OutActualFormat, uint32 SizeX, uint32 SizeY, uint32 ArraySize, uint32 NumMips, uint32 NumSamples, ETextureCreateFlags CreateFlags, const FClearValueBinding& ClearValueBinding) override final
+	virtual FXRSwapChainPtr CreateSwapchain(XrSession InSession, uint8 Format, uint8& OutActualFormat, uint32 SizeX, uint32 SizeY, uint32 ArraySize, uint32 NumMips, uint32 NumSamples, ETextureCreateFlags CreateFlags, const FClearValueBinding& ClearValueBinding, ETextureCreateFlags AuxiliaryCreateFlags) override final
 	{
-		return CreateSwapchain_OpenGL(InSession, Format, OutActualFormat, SizeX, SizeY, ArraySize, NumMips, NumSamples, CreateFlags, ClearValueBinding);
+		return CreateSwapchain_OpenGL(InSession, Format, OutActualFormat, SizeX, SizeY, ArraySize, NumMips, NumSamples, CreateFlags, ClearValueBinding, AuxiliaryCreateFlags);
 	}
 
 private:
@@ -224,9 +224,9 @@ public:
 		return nullptr;
 	}
 
-	virtual FXRSwapChainPtr CreateSwapchain(XrSession InSession, uint8 Format, uint8& OutActualFormat, uint32 SizeX, uint32 SizeY, uint32 ArraySize, uint32 NumMips, uint32 NumSamples, ETextureCreateFlags CreateFlags, const FClearValueBinding& ClearValueBinding) override final
+	virtual FXRSwapChainPtr CreateSwapchain(XrSession InSession, uint8 Format, uint8& OutActualFormat, uint32 SizeX, uint32 SizeY, uint32 ArraySize, uint32 NumMips, uint32 NumSamples, ETextureCreateFlags CreateFlags, const FClearValueBinding& ClearValueBinding, ETextureCreateFlags AuxiliaryCreateFlags) override final
 	{
-		return CreateSwapchain_OpenGLES(InSession, Format, OutActualFormat, SizeX, SizeY, ArraySize, NumMips, NumSamples, CreateFlags, ClearValueBinding);
+		return CreateSwapchain_OpenGLES(InSession, Format, OutActualFormat, SizeX, SizeY, ArraySize, NumMips, NumSamples, CreateFlags, ClearValueBinding, AuxiliaryCreateFlags);
 	}
 
 private:
@@ -303,9 +303,9 @@ public:
 		return GetIVulkanDynamicRHI()->RHIGetGraphicsAdapterLUID(Gpu);
 	}
 
-	virtual FXRSwapChainPtr CreateSwapchain(XrSession InSession, uint8 Format, uint8& OutActualFormat, uint32 SizeX, uint32 SizeY, uint32 ArraySize, uint32 NumMips, uint32 NumSamples, ETextureCreateFlags CreateFlags, const FClearValueBinding& ClearValueBinding) override final
+	virtual FXRSwapChainPtr CreateSwapchain(XrSession InSession, uint8 Format, uint8& OutActualFormat, uint32 SizeX, uint32 SizeY, uint32 ArraySize, uint32 NumMips, uint32 NumSamples, ETextureCreateFlags CreateFlags, const FClearValueBinding& ClearValueBinding, ETextureCreateFlags AuxiliaryCreateFlags) override final
 	{
-		return CreateSwapchain_Vulkan(InSession, Format, OutActualFormat, SizeX, SizeY, ArraySize, NumMips, NumSamples, CreateFlags, ClearValueBinding);
+		return CreateSwapchain_Vulkan(InSession, Format, OutActualFormat, SizeX, SizeY, ArraySize, NumMips, NumSamples, CreateFlags, ClearValueBinding, AuxiliaryCreateFlags);
 	}
 
 private:
diff --git a/Engine/Plugins/Runtime/OpenXR/Source/OpenXRHMD/Private/OpenXRHMD_RenderBridge.h b/Engine/Plugins/Runtime/OpenXR/Source/OpenXRHMD/Private/OpenXRHMD_RenderBridge.h
--- a/Engine/Plugins/Runtime/OpenXR/Source/OpenXRHMD/Private/OpenXRHMD_RenderBridge.h
+++ b/Engine/Plugins/Runtime/OpenXR/Source/OpenXRHMD/Private/OpenXRHMD_RenderBridge.h
@@ -25,7 +25,7 @@ public:
 	virtual uint64 GetGraphicsAdapterLuid(XrSystemId InSystem) { return 0; };
 	virtual void* GetGraphicsBinding(XrSystemId InSystem) = 0;
 
-	virtual FXRSwapChainPtr CreateSwapchain(XrSession InSession, uint8 Format, uint8& OutActualFormat, uint32 SizeX, uint32 SizeY, uint32 ArraySize, uint32 NumMips, uint32 NumSamples, ETextureCreateFlags CreateFlags, const FClearValueBinding& ClearValueBinding) = 0;
+	virtual FXRSwapChainPtr CreateSwapchain(XrSession InSession, uint8 Format, uint8& OutActualFormat, uint32 SizeX, uint32 SizeY, uint32 ArraySize, uint32 NumMips, uint32 NumSamples, ETextureCreateFlags CreateFlags, const FClearValueBinding& ClearValueBinding, ETextureCreateFlags AuxiliaryCreateFlags = ETextureCreateFlags::None) = 0;
 
 	FXRSwapChainPtr CreateSwapchain(XrSession InSession, FRHITexture2D* Template, ETextureCreateFlags CreateFlags)
 	{
diff --git a/Engine/Plugins/Runtime/OpenXR/Source/OpenXRHMD/Private/OpenXRHMD_Swapchain.cpp b/Engine/Plugins/Runtime/OpenXR/Source/OpenXRHMD/Private/OpenXRHMD_Swapchain.cpp
--- a/Engine/Plugins/Runtime/OpenXR/Source/OpenXRHMD/Private/OpenXRHMD_Swapchain.cpp
+++ b/Engine/Plugins/Runtime/OpenXR/Source/OpenXRHMD/Private/OpenXRHMD_Swapchain.cpp
@@ -234,7 +234,7 @@ XrSwapchain FOpenXRSwapchain::CreateSwapchain(XrSession InSession, uint32 Platfo
 }
 
 template<typename T>
-TArray<T> EnumerateImages(XrSwapchain InSwapchain, XrStructureType InType)
+TArray<T> EnumerateImages(XrSwapchain InSwapchain, XrStructureType InType, void* Next = nullptr)
 {
 	TArray<T> Images;
 	uint32_t ChainCount;
@@ -243,13 +243,45 @@ TArray<T> EnumerateImages(XrSwapchain InSwapchain, XrStructureType InType)
 	for (auto& Image : Images)
 	{
 		Image.type = InType;
+		Image.next = Next;
 	}
 	XR_ENSURE(xrEnumerateSwapchainImages(InSwapchain, ChainCount, &ChainCount, reinterpret_cast<XrSwapchainImageBaseHeader*>(Images.GetData())));
 	return Images;
 }
 
+void FOpenXRSwapchain::GetFragmentDensityMaps(TArray<FTextureRHIRef>& OutTextureChain, const bool bIsMobileMultiViewEnabled)
+{
+#ifdef XR_USE_GRAPHICS_API_VULKAN
+	IVulkanDynamicRHI* VulkanRHI = GetIVulkanDynamicRHI();
+
+	XrSwapchainImageFoveationVulkanFB FoveationImage{ XR_TYPE_SWAPCHAIN_IMAGE_FOVEATION_VULKAN_FB };
+	FoveationImage.next = nullptr;
+	void* Next = &FoveationImage;
+
+	XrSwapchain Swapchain = GetHandle();
+
+	TArray<XrSwapchainImageVulkanKHR> Images = EnumerateImages<XrSwapchainImageVulkanKHR>(Swapchain, XR_TYPE_SWAPCHAIN_IMAGE_VULKAN_KHR, Next);
+	const VkImageSubresourceRange SubresourceRangeAll = { VK_IMAGE_ASPECT_COLOR_BIT, 0, VK_REMAINING_MIP_LEVELS, 0, VK_REMAINING_ARRAY_LAYERS };
+
+	for (const XrSwapchainImageVulkanKHR& Image : Images)
+	{
+		const XrBaseOutStructure* NextHeader = reinterpret_cast<const XrBaseOutStructure*>(Image.next);
+		const XrSwapchainImageFoveationVulkanFB* FoveationImageVulkan = reinterpret_cast<const XrSwapchainImageFoveationVulkanFB*>(NextHeader);
+		// TODO: remove RHISetImageLayout when infrastructure to set an initial Vulkan layout is added.
+		VulkanRHI->RHISetImageLayout((VkImage)FoveationImageVulkan->image, VK_IMAGE_LAYOUT_UNDEFINED, VK_IMAGE_LAYOUT_FRAGMENT_DENSITY_MAP_OPTIMAL_EXT, SubresourceRangeAll);
+		OutTextureChain.Add(static_cast<FTextureRHIRef>(bIsMobileMultiViewEnabled ?
+			VulkanRHI->RHICreateTexture2DArrayFromResource(GPixelFormats[PF_R8G8].UnrealFormat, FoveationImageVulkan->width, FoveationImageVulkan->height, 2, 1, 1, FoveationImageVulkan->image, TexCreate_Foveation, FClearValueBinding::White) :
+			VulkanRHI->RHICreateTexture2DFromResource(GPixelFormats[PF_R8G8].UnrealFormat, FoveationImageVulkan->width, FoveationImageVulkan->height, 1, 1, FoveationImageVulkan->image, TexCreate_Foveation, FClearValueBinding::White)
+			));
+	}
+#else
+	OutTextureChain.Empty();
+#endif
+}
+
+
 #ifdef XR_USE_GRAPHICS_API_D3D11
-FXRSwapChainPtr CreateSwapchain_D3D11(XrSession InSession, uint8 Format, uint8& OutActualFormat, uint32 SizeX, uint32 SizeY, uint32 ArraySize, uint32 NumMips, uint32 NumSamples, ETextureCreateFlags CreateFlags, const FClearValueBinding& ClearValueBinding)
+FXRSwapChainPtr CreateSwapchain_D3D11(XrSession InSession, uint8 Format, uint8& OutActualFormat, uint32 SizeX, uint32 SizeY, uint32 ArraySize, uint32 NumMips, uint32 NumSamples, ETextureCreateFlags CreateFlags, const FClearValueBinding& ClearValueBinding, ETextureCreateFlags AuxiliaryCreateFlags)
 {
 	TFunction<uint32(uint8)> ToPlatformFormat = [](uint8 InFormat)
 	{
@@ -285,7 +317,7 @@ FXRSwapChainPtr CreateSwapchain_D3D11(XrSession InSession, uint8 Format, uint8&
 #endif
 
 #ifdef XR_USE_GRAPHICS_API_D3D12
-FXRSwapChainPtr CreateSwapchain_D3D12(XrSession InSession, uint8 Format, uint8& OutActualFormat, uint32 SizeX, uint32 SizeY, uint32 ArraySize, uint32 NumMips, uint32 NumSamples, ETextureCreateFlags CreateFlags, const FClearValueBinding& ClearValueBinding)
+FXRSwapChainPtr CreateSwapchain_D3D12(XrSession InSession, uint8 Format, uint8& OutActualFormat, uint32 SizeX, uint32 SizeY, uint32 ArraySize, uint32 NumMips, uint32 NumSamples, ETextureCreateFlags CreateFlags, const FClearValueBinding& ClearValueBinding, ETextureCreateFlags AuxiliaryCreateFlags)
 {
 	TFunction<uint32(uint8)> ToPlatformFormat = [](uint8 InFormat)
 	{
@@ -320,7 +352,7 @@ FXRSwapChainPtr CreateSwapchain_D3D12(XrSession InSession, uint8 Format, uint8&
 #endif
 
 #ifdef XR_USE_GRAPHICS_API_OPENGL
-FXRSwapChainPtr CreateSwapchain_OpenGL(XrSession InSession, uint8 Format, uint8& OutActualFormat, uint32 SizeX, uint32 SizeY, uint32 ArraySize, uint32 NumMips, uint32 NumSamples, ETextureCreateFlags CreateFlags, const FClearValueBinding& ClearValueBinding)
+FXRSwapChainPtr CreateSwapchain_OpenGL(XrSession InSession, uint8 Format, uint8& OutActualFormat, uint32 SizeX, uint32 SizeY, uint32 ArraySize, uint32 NumMips, uint32 NumSamples, ETextureCreateFlags CreateFlags, const FClearValueBinding& ClearValueBinding, ETextureCreateFlags AuxiliaryCreateFlags)
 {
 	Format = FOpenXRSwapchain::GetNearestSupportedSwapchainFormat(InSession, Format);
 	if (!Format)
@@ -350,7 +382,7 @@ FXRSwapChainPtr CreateSwapchain_OpenGL(XrSession InSession, uint8 Format, uint8&
 #endif
 
 #ifdef XR_USE_GRAPHICS_API_OPENGL_ES
-FXRSwapChainPtr CreateSwapchain_OpenGLES(XrSession InSession, uint8 Format, uint8& OutActualFormat, uint32 SizeX, uint32 SizeY, uint32 ArraySize, uint32 NumMips, uint32 NumSamples, ETextureCreateFlags CreateFlags, const FClearValueBinding& ClearValueBinding)
+FXRSwapChainPtr CreateSwapchain_OpenGLES(XrSession InSession, uint8 Format, uint8& OutActualFormat, uint32 SizeX, uint32 SizeY, uint32 ArraySize, uint32 NumMips, uint32 NumSamples, ETextureCreateFlags CreateFlags, const FClearValueBinding& ClearValueBinding, ETextureCreateFlags AuxiliaryCreateFlags)
 {
 	Format = FOpenXRSwapchain::GetNearestSupportedSwapchainFormat(InSession, Format);
 	if (!Format)
@@ -359,7 +391,17 @@ FXRSwapChainPtr CreateSwapchain_OpenGLES(XrSession InSession, uint8 Format, uint
 	}
 
 	OutActualFormat = Format;
-	XrSwapchain Swapchain = FOpenXRSwapchain::CreateSwapchain(InSession, GPixelFormats[Format].PlatformFormat, SizeX, SizeY, ArraySize, NumMips, NumSamples, CreateFlags);
+
+	void* Next = nullptr;
+	XrSwapchainCreateInfoFoveationFB FoveationCreateInfo{ XR_TYPE_SWAPCHAIN_CREATE_INFO_FOVEATION_FB };
+	if (EnumHasAllFlags(AuxiliaryCreateFlags, TexCreate_Foveation))
+	{
+		FoveationCreateInfo.next = Next;
+		FoveationCreateInfo.flags = XR_SWAPCHAIN_CREATE_FOVEATION_SCALED_BIN_BIT_FB;
+
+		Next = &FoveationCreateInfo;
+	}
+	XrSwapchain Swapchain = FOpenXRSwapchain::CreateSwapchain(InSession, GPixelFormats[Format].PlatformFormat, SizeX, SizeY, ArraySize, NumMips, NumSamples, CreateFlags, Next);
 	if (!Swapchain)
 	{
 		return nullptr;
@@ -380,7 +422,7 @@ FXRSwapChainPtr CreateSwapchain_OpenGLES(XrSession InSession, uint8 Format, uint
 #endif
 
 #ifdef XR_USE_GRAPHICS_API_VULKAN
-FXRSwapChainPtr CreateSwapchain_Vulkan(XrSession InSession, uint8 Format, uint8& OutActualFormat, uint32 SizeX, uint32 SizeY, uint32 ArraySize, uint32 NumMips, uint32 NumSamples, ETextureCreateFlags CreateFlags, const FClearValueBinding& ClearValueBinding)
+FXRSwapChainPtr CreateSwapchain_Vulkan(XrSession InSession, uint8 Format, uint8& OutActualFormat, uint32 SizeX, uint32 SizeY, uint32 ArraySize, uint32 NumMips, uint32 NumSamples, ETextureCreateFlags CreateFlags, const FClearValueBinding& ClearValueBinding, ETextureCreateFlags AuxiliaryCreateFlags)
 {
 	TFunction<uint32(uint8)> ToPlatformFormat = [](uint8 InFormat)
 	{
@@ -409,6 +451,15 @@ FXRSwapChainPtr CreateSwapchain_Vulkan(XrSession InSession, uint8 Format, uint8&
 	// We don't need to specify additional formats if the swapchain is being created as sRGB.
 	void* Next = !(CreateFlags & TexCreate_SRGB) ? &FormatListInfo : nullptr;
 
+	XrSwapchainCreateInfoFoveationFB FoveationCreateInfo { XR_TYPE_SWAPCHAIN_CREATE_INFO_FOVEATION_FB };
+	if (EnumHasAllFlags(AuxiliaryCreateFlags, TexCreate_Foveation))
+	{
+		FoveationCreateInfo.next = Next;
+		FoveationCreateInfo.flags = XR_SWAPCHAIN_CREATE_FOVEATION_FRAGMENT_DENSITY_MAP_BIT_FB;
+
+		Next = &FoveationCreateInfo;
+	}
+
 	XrSwapchain Swapchain = FOpenXRSwapchain::CreateSwapchain(InSession, ViewFormatList[0], SizeX, SizeY, ArraySize, NumMips, NumSamples, CreateFlags, Next);
 	if (!Swapchain)
 	{
diff --git a/Engine/Plugins/Runtime/OpenXR/Source/OpenXRHMD/Private/OpenXRHMD_Swapchain.h b/Engine/Plugins/Runtime/OpenXR/Source/OpenXRHMD/Private/OpenXRHMD_Swapchain.h
--- a/Engine/Plugins/Runtime/OpenXR/Source/OpenXRHMD/Private/OpenXRHMD_Swapchain.h
+++ b/Engine/Plugins/Runtime/OpenXR/Source/OpenXRHMD/Private/OpenXRHMD_Swapchain.h
@@ -22,6 +22,7 @@ public:
 	XrSwapchain GetHandle() { return Handle; }
 	static XrSwapchain CreateSwapchain(XrSession InSession, uint32 PlatformFormat, uint32 SizeX, uint32 SizeY, uint32 ArraySize, uint32 NumMips, uint32 NumSamples, ETextureCreateFlags CreateFlags, void* Next = nullptr);
 	static uint8 GetNearestSupportedSwapchainFormat(XrSession InSession, uint8 RequestedFormat, TFunction<uint32(uint8)> ToPlatformFormat = nullptr);
+	void GetFragmentDensityMaps(TArray<FTextureRHIRef>& OutTextureChain, const bool bIsMobileMultiViewEnabled);
 
 protected:
 	XrSwapchain Handle;
@@ -34,17 +35,17 @@ protected:
 };
 
 #ifdef XR_USE_GRAPHICS_API_D3D11
-FXRSwapChainPtr CreateSwapchain_D3D11(XrSession InSession, uint8 Format, uint8& OutActualFormat, uint32 SizeX, uint32 SizeY, uint32 ArraySize, uint32 NumMips, uint32 NumSamples, ETextureCreateFlags CreateFlags, const FClearValueBinding& ClearValueBinding);
+FXRSwapChainPtr CreateSwapchain_D3D11(XrSession InSession, uint8 Format, uint8& OutActualFormat, uint32 SizeX, uint32 SizeY, uint32 ArraySize, uint32 NumMips, uint32 NumSamples, ETextureCreateFlags CreateFlags, const FClearValueBinding& ClearValueBinding, ETextureCreateFlags AuxiliaryCreateFlags = ETextureCreateFlags::None);
 #endif
 #ifdef XR_USE_GRAPHICS_API_D3D12
-FXRSwapChainPtr CreateSwapchain_D3D12(XrSession InSession, uint8 Format, uint8& OutActualFormat, uint32 SizeX, uint32 SizeY, uint32 ArraySize, uint32 NumMips, uint32 NumSamples, ETextureCreateFlags CreateFlags, const FClearValueBinding& ClearValueBinding);
+FXRSwapChainPtr CreateSwapchain_D3D12(XrSession InSession, uint8 Format, uint8& OutActualFormat, uint32 SizeX, uint32 SizeY, uint32 ArraySize, uint32 NumMips, uint32 NumSamples, ETextureCreateFlags CreateFlags, const FClearValueBinding& ClearValueBinding, ETextureCreateFlags AuxiliaryCreateFlags = ETextureCreateFlags::None);
 #endif
 #ifdef XR_USE_GRAPHICS_API_OPENGL
-FXRSwapChainPtr CreateSwapchain_OpenGL(XrSession InSession, uint8 Format, uint8& OutActualFormat, uint32 SizeX, uint32 SizeY, uint32 ArraySize, uint32 NumMips, uint32 NumSamples, ETextureCreateFlags CreateFlags, const FClearValueBinding& ClearValueBinding);
+FXRSwapChainPtr CreateSwapchain_OpenGL(XrSession InSession, uint8 Format, uint8& OutActualFormat, uint32 SizeX, uint32 SizeY, uint32 ArraySize, uint32 NumMips, uint32 NumSamples, ETextureCreateFlags CreateFlags, const FClearValueBinding& ClearValueBinding, ETextureCreateFlags AuxiliaryCreateFlags = ETextureCreateFlags::None);
 #endif
 #ifdef XR_USE_GRAPHICS_API_OPENGL_ES
-FXRSwapChainPtr CreateSwapchain_OpenGLES(XrSession InSession, uint8 Format, uint8& OutActualFormat, uint32 SizeX, uint32 SizeY, uint32 ArraySize, uint32 NumMips, uint32 NumSamples, ETextureCreateFlags CreateFlags, const FClearValueBinding& ClearValueBinding);
+FXRSwapChainPtr CreateSwapchain_OpenGLES(XrSession InSession, uint8 Format, uint8& OutActualFormat, uint32 SizeX, uint32 SizeY, uint32 ArraySize, uint32 NumMips, uint32 NumSamples, ETextureCreateFlags CreateFlags, const FClearValueBinding& ClearValueBinding, ETextureCreateFlags AuxiliaryCreateFlags = ETextureCreateFlags::None);
 #endif
 #ifdef XR_USE_GRAPHICS_API_VULKAN
-FXRSwapChainPtr CreateSwapchain_Vulkan(XrSession InSession, uint8 Format, uint8& OutActualFormat, uint32 SizeX, uint32 SizeY, uint32 ArraySize, uint32 NumMips, uint32 NumSamples, ETextureCreateFlags CreateFlags, const FClearValueBinding& ClearValueBinding);
+FXRSwapChainPtr CreateSwapchain_Vulkan(XrSession InSession, uint8 Format, uint8& OutActualFormat, uint32 SizeX, uint32 SizeY, uint32 ArraySize, uint32 NumMips, uint32 NumSamples, ETextureCreateFlags CreateFlags, const FClearValueBinding& ClearValueBinding, ETextureCreateFlags AuxiliaryCreateFlags = ETextureCreateFlags::None);
 #endif
\ No newline at end of file
