Wrap OpenXR uses of UPlayerMappableInputConfig asset types in disable deprecation pragmas. This asset type will be deprecated and should be updated to use the new Enhanced Input User Settings for 5.4.

diff --git a/Engine/Plugins/Experimental/XRCreativeFramework/Source/XRCreativeEditor/Private/XRCreativeVREditorMode.cpp b/Engine/Plugins/Experimental/XRCreativeFramework/Source/XRCreativeEditor/Private/XRCreativeVREditorMode.cpp
--- a/Engine/Plugins/Experimental/XRCreativeFramework/Source/XRCreativeEditor/Private/XRCreativeVREditorMode.cpp
+++ b/Engine/Plugins/Experimental/XRCreativeFramework/Source/XRCreativeEditor/Private/XRCreativeVREditorMode.cpp
@@ -218,7 +218,9 @@ void UXRCreativeVREditorMode::EnableStereo()
 		Viewport->SetRenderDirectlyToWindow(true);
 	}
 
+	PRAGMA_DISABLE_DEPRECATION_WARNINGS
 	UOpenXRInputFunctionLibrary::BeginXRSession(Cast<UPlayerMappableInputConfig>(MappableInputConfig.TryLoad()));
+	PRAGMA_ENABLE_DEPRECATION_WARNINGS
 }
 
 
diff --git a/Engine/Plugins/Runtime/OpenXR/Source/OpenXRInput/Private/OpenXRInput.cpp b/Engine/Plugins/Runtime/OpenXR/Source/OpenXRInput/Private/OpenXRInput.cpp
--- a/Engine/Plugins/Runtime/OpenXR/Source/OpenXRInput/Private/OpenXRInput.cpp
+++ b/Engine/Plugins/Runtime/OpenXR/Source/OpenXRInput/Private/OpenXRInput.cpp
@@ -372,6 +372,7 @@ bool FOpenXRInputPlugin::FOpenXRInput::BuildActions(XrSession Session)
 		}
 	}
 
+	PRAGMA_DISABLE_DEPRECATION_WARNINGS
 	if (!MappableInputConfig)
 	{
 		// Attempt to load the default input config from the OpenXR input settings.
@@ -392,6 +393,7 @@ bool FOpenXRInputPlugin::FOpenXRInput::BuildActions(XrSession Session)
 
 		BuildLegacyActions(Profiles);
 	}
+	PRAGMA_ENABLE_DEPRECATION_WARNINGS
 
 	for (TPair<FString, FInteractionProfile>& Pair : Profiles)
 	{
@@ -550,6 +552,7 @@ void FOpenXRInputPlugin::FOpenXRInput::BuildEnhancedActions(TMap<FString, FInter
 		return;
 	}
 
+	PRAGMA_DISABLE_DEPRECATION_WARNINGS
 	for (const TPair<TObjectPtr<UInputMappingContext>, int32> MappingContext : MappableInputConfig->GetMappingContexts())
 	{
 		FOpenXRActionSet ActionSet(Instance, MappingContext.Key->GetFName(), MappingContext.Key->ContextDescription.ToString(), ToXrPriority(MappingContext.Value), MappingContext.Key);
@@ -583,6 +586,7 @@ void FOpenXRInputPlugin::FOpenXRInput::BuildEnhancedActions(TMap<FString, FInter
 		}
 		ActionSets.Emplace(MoveTemp(ActionSet));
 	}
+	PRAGMA_ENABLE_DEPRECATION_WARNINGS
 }
 
 void FOpenXRInputPlugin::FOpenXRInput::DestroyActions()
@@ -1384,6 +1388,7 @@ float FOpenXRInputPlugin::FOpenXRInput::GetHapticAmplitudeScale() const
 	return 1.0f;
 }
 
+PRAGMA_DISABLE_DEPRECATION_WARNINGS
 bool FOpenXRInputPlugin::FOpenXRInput::SetPlayerMappableInputConfig(TObjectPtr<class UPlayerMappableInputConfig> InputConfig)
 {
 	if (bActionsAttached)
@@ -1396,5 +1401,6 @@ bool FOpenXRInputPlugin::FOpenXRInput::SetPlayerMappableInputConfig(TObjectPtr<c
 	MappableInputConfig = TStrongObjectPtr<class UPlayerMappableInputConfig>(InputConfig);
 	return true;
 }
+PRAGMA_ENABLE_DEPRECATION_WARNINGS
 
 #undef LOCTEXT_NAMESPACE // "OpenXRInputPlugin"
diff --git a/Engine/Plugins/Runtime/OpenXR/Source/OpenXRInput/Private/OpenXRInput.h b/Engine/Plugins/Runtime/OpenXR/Source/OpenXRInput/Private/OpenXRInput.h
--- a/Engine/Plugins/Runtime/OpenXR/Source/OpenXRInput/Private/OpenXRInput.h
+++ b/Engine/Plugins/Runtime/OpenXR/Source/OpenXRInput/Private/OpenXRInput.h
@@ -16,7 +16,6 @@ class UInputAction;
 class UInputTrigger;
 class UInputModifier;
 class UInputMappingContext;
-class UPlayerMappableInputConfig;
 struct FInputActionKeyMapping;
 struct FInputAxisKeyMapping;
 struct FKey;
@@ -127,7 +126,9 @@ public:
 		virtual bool GetControllerOrientationAndPositionForTime(const int32 ControllerIndex, const FName MotionSource, FTimespan Time, bool& OutTimeWasUsed, FRotator& OutOrientation, FVector& OutPosition, bool& OutbProvidedLinearVelocity, FVector& OutLinearVelocity, bool& OutbProvidedAngularVelocity, FVector& OutAngularVelocityRadPerSec, bool& OutbProvidedLinearAcceleration, FVector& OutLinearAcceleration, float WorldToMetersScale) const override;
 		virtual ETrackingStatus GetControllerTrackingStatus(const int32 ControllerIndex, const FName MotionSource) const override;
 		virtual void EnumerateSources(TArray<FMotionControllerSource>& SourcesOut) const override;
+		PRAGMA_DISABLE_DEPRECATION_WARNINGS
 		virtual bool SetPlayerMappableInputConfig(TObjectPtr<class UPlayerMappableInputConfig> InputConfig) override;
+		PRAGMA_ENABLE_DEPRECATION_WARNINGS
 
 		// IHapticDevice overrides
 		IHapticDevice* GetHapticDevice() override { return (IHapticDevice*)this; }
@@ -147,7 +148,9 @@ public:
 		TMap<EControllerHand, FOpenXRController> Controllers;
 		TMap<FName, EControllerHand> MotionSourceToControllerHandMap;
 
-		TStrongObjectPtr<UPlayerMappableInputConfig> MappableInputConfig;
+		PRAGMA_DISABLE_DEPRECATION_WARNINGS
+		TStrongObjectPtr<class UPlayerMappableInputConfig> MappableInputConfig;
+		PRAGMA_ENABLE_DEPRECATION_WARNINGS
 
 		XrAction GetActionForMotionSource(FName MotionSource) const;
 		int32 GetDeviceIDForMotionSource(FName MotionSource) const;
diff --git a/Engine/Plugins/Runtime/OpenXR/Source/OpenXRInput/Private/OpenXRInputFunctionLibrary.cpp b/Engine/Plugins/Runtime/OpenXR/Source/OpenXRInput/Private/OpenXRInputFunctionLibrary.cpp
--- a/Engine/Plugins/Runtime/OpenXR/Source/OpenXRInput/Private/OpenXRInputFunctionLibrary.cpp
+++ b/Engine/Plugins/Runtime/OpenXR/Source/OpenXRInput/Private/OpenXRInputFunctionLibrary.cpp
@@ -15,6 +15,7 @@ UOpenXRInputFunctionLibrary::UOpenXRInputFunctionLibrary(const FObjectInitialize
 {
 }
 
+PRAGMA_DISABLE_DEPRECATION_WARNINGS
 bool UOpenXRInputFunctionLibrary::BeginXRSession(UPlayerMappableInputConfig* InputConfig)
 {
 	if (!InputConfig)
@@ -41,6 +42,7 @@ bool UOpenXRInputFunctionLibrary::BeginXRSession(UPlayerMappableInputConfig* Inp
 
 	return UHeadMountedDisplayFunctionLibrary::EnableHMD(true);
 }
+PRAGMA_ENABLE_DEPRECATION_WARNINGS
 
 void UOpenXRInputFunctionLibrary::EndXRSession()
 {
diff --git a/Engine/Plugins/Runtime/OpenXR/Source/OpenXRInput/Private/OpenXRInputSettings.cpp b/Engine/Plugins/Runtime/OpenXR/Source/OpenXRInput/Private/OpenXRInputSettings.cpp
--- a/Engine/Plugins/Runtime/OpenXR/Source/OpenXRInput/Private/OpenXRInputSettings.cpp
+++ b/Engine/Plugins/Runtime/OpenXR/Source/OpenXRInput/Private/OpenXRInputSettings.cpp
@@ -21,6 +21,7 @@ void UOpenXRInputSettings::PostEditChangeChainProperty(FPropertyChangedChainEven
 
 	const FName MemberPropertyName = PropertyChangedEvent.PropertyChain.GetActiveMemberNode()->GetValue()->GetFName();
 
+	PRAGMA_DISABLE_DEPRECATION_WARNINGS
 	if (MemberPropertyName == GET_MEMBER_NAME_CHECKED(UOpenXRInputSettings, MappableInputConfig))
 	{
 		TArray<IMotionController*> MotionControllers = IModularFeatures::Get().GetModularFeatureImplementations<IMotionController>(IMotionController::GetModularFeatureName());
@@ -34,5 +35,6 @@ void UOpenXRInputSettings::PostEditChangeChainProperty(FPropertyChangedChainEven
 			MotionController->SetPlayerMappableInputConfig((UPlayerMappableInputConfig*)MappableInputConfig.TryLoad());
 		}
 	}
+	PRAGMA_ENABLE_DEPRECATION_WARNINGS
 }
 #endif
diff --git a/Engine/Plugins/Runtime/OpenXR/Source/OpenXRInput/Public/OpenXRInputFunctionLibrary.h b/Engine/Plugins/Runtime/OpenXR/Source/OpenXRInput/Public/OpenXRInputFunctionLibrary.h
--- a/Engine/Plugins/Runtime/OpenXR/Source/OpenXRInput/Public/OpenXRInputFunctionLibrary.h
+++ b/Engine/Plugins/Runtime/OpenXR/Source/OpenXRInput/Public/OpenXRInputFunctionLibrary.h
@@ -12,6 +12,7 @@ class OPENXRINPUT_API UOpenXRInputFunctionLibrary : public UBlueprintFunctionLib
 {
 	GENERATED_UCLASS_BODY()
 
+	PRAGMA_DISABLE_DEPRECATION_WARNINGS
 	/**
 	* Add a player mappable input config for Enhanced Input. This will make the tracking system aware of the actions and bindings that are used for XR motion controllers.
 	* Attaching input configs to the session can only be done once, so if multiple input configs need to be added only check AttachToSession for the last input config.
@@ -23,6 +24,7 @@ class OPENXRINPUT_API UOpenXRInputFunctionLibrary : public UBlueprintFunctionLib
 	*/
 	UFUNCTION(BlueprintCallable, Category = "Input|XRTracking")
 	static bool BeginXRSession(class UPlayerMappableInputConfig* InputConfig);
+	PRAGMA_ENABLE_DEPRECATION_WARNINGS
 
 
 	UFUNCTION(BlueprintCallable, Category = "Input|XRTracking")
diff --git a/Engine/Plugins/Runtime/OpenXR/Source/OpenXRInput/Public/OpenXRInputSettings.h b/Engine/Plugins/Runtime/OpenXR/Source/OpenXRInput/Public/OpenXRInputSettings.h
--- a/Engine/Plugins/Runtime/OpenXR/Source/OpenXRInput/Public/OpenXRInputSettings.h
+++ b/Engine/Plugins/Runtime/OpenXR/Source/OpenXRInput/Public/OpenXRInputSettings.h
@@ -15,9 +15,11 @@ class OPENXRINPUT_API UOpenXRInputSettings : public UObject
 	GENERATED_UCLASS_BODY()
 
 public:
+	PRAGMA_DISABLE_DEPRECATION_WARNINGS
 	/** Set a mappable input config to allow OpenXR runtimes to remap the Enhanced Input actions. */
 	UPROPERTY(config, EditAnywhere, Category = "Enhanced Input", meta = (DisplayName = "Mappable Input Config for XR", AllowedClasses = "/Script/EnhancedInput.PlayerMappableInputConfig"))
 	FSoftObjectPath MappableInputConfig = nullptr;
+	PRAGMA_ENABLE_DEPRECATION_WARNINGS
 
 	// UObject interface
 #if WITH_EDITOR
