Remove xrSyncAction log spam.

diff --git a/Engine/Plugins/Runtime/OpenXR/Source/OpenXRInput/Private/OpenXRInput.cpp b/Engine/Plugins/Runtime/OpenXR/Source/OpenXRInput/Private/OpenXRInput.cpp
--- a/Engine/Plugins/Runtime/OpenXR/Source/OpenXRInput/Private/OpenXRInput.cpp
+++ b/Engine/Plugins/Runtime/OpenXR/Source/OpenXRInput/Private/OpenXRInput.cpp
@@ -241,6 +241,7 @@ FOpenXRInputPlugin::FOpenXRInput::FOpenXRInput(FOpenXRHMD* HMD)
 	, bDirectionalBindingSupported(false)
 	, bPalmPoseSupported(false)
 	, MessageHandler(new FGenericApplicationMessageHandler())
+	, bActionSetPrioritySupported(false)
 {
 	IModularFeatures::Get().RegisterModularFeature(GetModularFeatureName(), this);
 	
@@ -250,6 +251,7 @@ FOpenXRInputPlugin::FOpenXRInput::FOpenXRInput(FOpenXRHMD* HMD)
 		Instance = OpenXRHMD->GetInstance();
 		bDirectionalBindingSupported = OpenXRHMD->IsExtensionEnabled(XR_EXT_DPAD_BINDING_EXTENSION_NAME);
 		bPalmPoseSupported = OpenXRHMD->IsExtensionEnabled(XR_EXT_PALM_POSE_EXTENSION_NAME);
+		bActionSetPrioritySupported = OpenXRHMD->IsExtensionEnabled(XR_EXT_ACTIVE_ACTION_SET_PRIORITY_EXTENSION_NAME);
 
 		// Note: AnyHand needs special handling because it tries left then falls back to right in each call.
 		MotionSourceToControllerHandMap.Add(OpenXRSourceNames::Left, EControllerHand::Left);
@@ -809,7 +811,13 @@ void FOpenXRInputPlugin::FOpenXRInput::SyncActions(XrSession Session)
 
 		XrActionsSyncInfo SyncInfo = { XR_TYPE_ACTIONS_SYNC_INFO };
 		XrActiveActionSetPrioritiesEXT PrioritiesInfo = { XR_TYPE_ACTIVE_ACTION_SET_PRIORITIES_EXT };
-		SyncInfo.next = &PrioritiesInfo;
+		// This is a workaround to avoid log spam from some OpenXR runtimes when they do not support this extension.
+		if (bActionSetPrioritySupported)
+		{
+			SyncInfo.next = &PrioritiesInfo;
+			PrioritiesInfo.actionSetPriorityCount = ActivePriorities.Num();
+			PrioritiesInfo.actionSetPriorities = ActivePriorities.GetData();
+		}
 
 		for (IOpenXRExtensionPlugin* Plugin : OpenXRHMD->GetExtensionPlugins())
 		{
@@ -818,8 +826,6 @@ void FOpenXRInputPlugin::FOpenXRInput::SyncActions(XrSession Session)
 			SyncInfo.next = Plugin->OnSyncActions(Session, SyncInfo.next);
 		}
 
-		PrioritiesInfo.actionSetPriorityCount = ActivePriorities.Num();
-		PrioritiesInfo.actionSetPriorities = ActivePriorities.GetData();
 		SyncInfo.countActiveActionSets = ActiveActionSets.Num();
 		SyncInfo.activeActionSets = ActiveActionSets.GetData();
 
diff --git a/Engine/Plugins/Runtime/OpenXR/Source/OpenXRInput/Private/OpenXRInput.h b/Engine/Plugins/Runtime/OpenXR/Source/OpenXRInput/Private/OpenXRInput.h
--- a/Engine/Plugins/Runtime/OpenXR/Source/OpenXRInput/Private/OpenXRInput.h
+++ b/Engine/Plugins/Runtime/OpenXR/Source/OpenXRInput/Private/OpenXRInput.h
@@ -160,6 +160,7 @@ public:
 		bool bActionsAttached;
 		bool bDirectionalBindingSupported;
 		bool bPalmPoseSupported;
+		bool bActionSetPrioritySupported;
 
 		/**
 		* Buffer for current delta time to get an accurate approximation of how long to play haptics for
