Add setting to enable/disable OpenXR foveation.

diff --git a/Engine/Plugins/Runtime/OpenXR/Source/OpenXREditor/Private/OpenXREditorModule.cpp b/Engine/Plugins/Runtime/OpenXR/Source/OpenXREditor/Private/OpenXREditorModule.cpp
--- a/Engine/Plugins/Runtime/OpenXR/Source/OpenXREditor/Private/OpenXREditorModule.cpp
+++ b/Engine/Plugins/Runtime/OpenXR/Source/OpenXREditor/Private/OpenXREditorModule.cpp
@@ -3,6 +3,7 @@
 #include "OpenXREditorModule.h"
 #include "OpenXRAssetDirectory.h"
 #include "OpenXRInputSettings.h"
+#include "OpenXRHMDSettings.h"
 
 #include "ISettingsModule.h"
 #include "Modules/ModuleManager.h"
@@ -25,6 +26,12 @@ void FOpenXREditorModule::StartupModule()
 			LOCTEXT("OpenXRInputSettingsDescription", "Project settings for OpenXR plugin"),
 			GetMutableDefault<UOpenXRInputSettings>()
 		);
+
+		SettingsModule->RegisterSettings("Project", "Plugins", "OpenXRHMD",
+			LOCTEXT("OpenXRHMDSettingsName", "OpenXR Settings"),
+			LOCTEXT("OpenXRHMDSettingsDescription", "Project settings for OpenXR plugin"),
+			GetMutableDefault<UOpenXRHMDSettings>()
+		);
 	}
 
 	UEditorPerformanceSettings* EditorPerformanceSettings = GetMutableDefault<UEditorPerformanceSettings>();
@@ -51,6 +58,7 @@ void FOpenXREditorModule::ShutdownModule()
 	if (SettingsModule != nullptr)
 	{
 		SettingsModule->UnregisterSettings("Project", "Plugins", "OpenXR");
+		SettingsModule->UnregisterSettings("Project", "Plugins", "OpenXRHMD");
 	}
 }
 
diff --git a/Engine/Plugins/Runtime/OpenXR/Source/OpenXRHMD/Private/OpenXRHMDSettings.cpp b/Engine/Plugins/Runtime/OpenXR/Source/OpenXRHMD/Private/OpenXRHMDSettings.cpp
new file mode 100644
--- /dev/null
+++ b/Engine/Plugins/Runtime/OpenXR/Source/OpenXRHMD/Private/OpenXRHMDSettings.cpp
@@ -0,0 +1,8 @@
+// Copyright Epic Games, Inc. All Rights Reserved.
+
+#include "OpenXRHMDSettings.h"
+
+UOpenXRHMDSettings::UOpenXRHMDSettings(const FObjectInitializer& ObjectInitializer)
+	: Super(ObjectInitializer)
+{
+}
\ No newline at end of file
diff --git a/Engine/Plugins/Runtime/OpenXR/Source/OpenXRHMD/Public/OpenXRHMDSettings.h b/Engine/Plugins/Runtime/OpenXR/Source/OpenXRHMD/Public/OpenXRHMDSettings.h
new file mode 100644
--- /dev/null
+++ b/Engine/Plugins/Runtime/OpenXR/Source/OpenXRHMD/Public/OpenXRHMDSettings.h
@@ -0,0 +1,21 @@
+// Copyright Epic Games, Inc. All Rights Reserved.
+
+#pragma once
+
+#include "CoreMinimal.h"
+#include "UObject/Object.h"
+#include "OpenXRHMDSettings.generated.h"
+
+/**
+* Implements the settings for the OpenXR HMD plugin.
+*/
+UCLASS(config = Game, defaultconfig)
+class OPENXRHMD_API UOpenXRHMDSettings : public UObject
+{
+	GENERATED_UCLASS_BODY()
+
+public:
+	/** Enables foveation provided by the XR_FB_foveation OpenXR extension. */
+	UPROPERTY(config, EditAnywhere, Category = "Foveation", meta = (DisplayName = "Enable XR_FB_foveation extension"))
+	bool bIsFBFoveationEnabled = false;
+};
