UE-190378 OpenXRHMD SystemProperties uninitialized until OnStereoStartup, GetHMDName could return garbage.
-Initialize it in the constructor, also made the hand tracking system properties into a member so our member doesn't end up pointing at a local variable that goes out of scope.

diff --git a/Engine/Plugins/Runtime/OpenXR/Source/OpenXRHMD/Private/OpenXRHMD.cpp b/Engine/Plugins/Runtime/OpenXR/Source/OpenXRHMD/Private/OpenXRHMD.cpp
--- a/Engine/Plugins/Runtime/OpenXR/Source/OpenXRHMD/Private/OpenXRHMD.cpp
+++ b/Engine/Plugins/Runtime/OpenXR/Source/OpenXRHMD/Private/OpenXRHMD.cpp
@@ -1328,6 +1328,9 @@ FOpenXRHMD::FOpenXRHMD(const FAutoRegister& AutoRegister, XrInstance InInstance,
 	XR_ENSURE(xrGetInstanceProperties(Instance, &InstanceProperties));
 	InstanceProperties.runtimeName[XR_MAX_RUNTIME_NAME_SIZE - 1] = 0; // Ensure the name is null terminated.
 
+	SystemHandTrackingProperties = { XR_TYPE_SYSTEM_HAND_TRACKING_PROPERTIES_EXT };
+	SystemProperties = XrSystemProperties{ XR_TYPE_SYSTEM_PROPERTIES, &SystemHandTrackingProperties };
+
 	bDepthExtensionSupported = IsExtensionEnabled(XR_KHR_COMPOSITION_LAYER_DEPTH_EXTENSION_NAME) && CheckPlatformDepthExtensionSupport(InstanceProperties);
 	bHiddenAreaMaskSupported = IsExtensionEnabled(XR_KHR_VISIBILITY_MASK_EXTENSION_NAME) &&
 		!FCStringAnsi::Strstr(InstanceProperties.runtimeName, "Oculus");
@@ -1723,10 +1726,8 @@ bool FOpenXRHMD::OnStereoStartup()
 	}
 
 	// Retrieve system properties and check for hand tracking support
-	XrSystemHandTrackingPropertiesEXT HandTrackingSystemProperties = { XR_TYPE_SYSTEM_HAND_TRACKING_PROPERTIES_EXT };
-	SystemProperties = XrSystemProperties{ XR_TYPE_SYSTEM_PROPERTIES, &HandTrackingSystemProperties };
 	XR_ENSURE(xrGetSystemProperties(Instance, System, &SystemProperties));
-	bSupportsHandTracking = HandTrackingSystemProperties.supportsHandTracking == XR_TRUE;
+	bSupportsHandTracking = SystemHandTrackingProperties.supportsHandTracking == XR_TRUE;
 
 	// Some runtimes aren't compliant with their number of layers supported.
 	// We support a fallback by emulating non-facelocked layers
diff --git a/Engine/Plugins/Runtime/OpenXR/Source/OpenXRHMD/Private/OpenXRHMD.h b/Engine/Plugins/Runtime/OpenXR/Source/OpenXRHMD/Private/OpenXRHMD.h
--- a/Engine/Plugins/Runtime/OpenXR/Source/OpenXRHMD/Private/OpenXRHMD.h
+++ b/Engine/Plugins/Runtime/OpenXR/Source/OpenXRHMD/Private/OpenXRHMD.h
@@ -477,6 +477,7 @@ private:
 	XrViewConfigurationType SelectedViewConfigurationType;
 	XrEnvironmentBlendMode  SelectedEnvironmentBlendMode;
 	XrInstanceProperties    InstanceProperties;
+	XrSystemHandTrackingPropertiesEXT SystemHandTrackingProperties;
 	XrSystemProperties      SystemProperties;
 
 	FPipelinedFrameState	PipelinedFrameStateGame;
