OpenXR: work around SteamVR instance failures

diff --git a/Engine/Plugins/Runtime/OpenXR/Source/OpenXRHMD/Private/OpenXRHMDModule.cpp b/Engine/Plugins/Runtime/OpenXR/Source/OpenXRHMD/Private/OpenXRHMDModule.cpp
--- a/Engine/Plugins/Runtime/OpenXR/Source/OpenXRHMD/Private/OpenXRHMDModule.cpp
+++ b/Engine/Plugins/Runtime/OpenXR/Source/OpenXRHMD/Private/OpenXRHMDModule.cpp
@@ -31,6 +31,12 @@ static TAutoConsoleVariable<bool> CVarDisableOpenXROnAndroidWithoutOculus(
 	TEXT("If true OpenXR will not initialize on Android unless the project is packaged for Oculus (ProjectSetting->Platforms->Android->Advanced APK Packaging->PackageForOculusMobileDevices list not empty).  Currently defaulted to true because the OpenXR loader we are using hangs during intialization on some devices instead of failing, as it should."),
 	ECVF_RenderThreadSafe);
 
+static TAutoConsoleVariable<bool> CVarCheckOpenXRInstanceConformance(
+	TEXT("xr.CheckOpenXRInstanceConformance"),
+	true,
+	TEXT("If true, OpenXR will verify Instance is conformant by calling xrStringToPath. Some runtimes fail without a system attached at instance creation time."),
+	ECVF_RenderThreadSafe);
+
 //---------------------------------------------------
 // OpenXRHMD Plugin Implementation
 //---------------------------------------------------
@@ -757,6 +763,22 @@ bool FOpenXRHMDModule::InitInstance()
 	InstanceProps.runtimeName[XR_MAX_RUNTIME_NAME_SIZE - 1] = 0; // Ensure the name is null terminated.
 	UE_LOG(LogHMD, Log, TEXT("Initialized OpenXR on %S runtime version %d.%d.%d"), InstanceProps.runtimeName, XR_VERSION_MAJOR(InstanceProps.runtimeVersion), XR_VERSION_MINOR(InstanceProps.runtimeVersion), XR_VERSION_PATCH(InstanceProps.runtimeVersion));
 
+	if (CVarCheckOpenXRInstanceConformance.GetValueOnAnyThread() &&
+		(FCStringAnsi::Strstr(InstanceProps.runtimeName, "SteamVR/OpenXR") != nullptr))
+	{
+		// Runtimes should not be dependent on system availability to use instance-only functions.
+		// However, some runtimes fail with some instance-only calls, such as xrStringToPath. We
+		// need to bail early to prevent failures later on during setup.
+		XrPath UserHeadTestPath = XR_NULL_PATH;
+		const XrResult StringToPathTest = xrStringToPath(Instance, "/user/head", &UserHeadTestPath);
+
+		if (StringToPathTest != XR_SUCCESS)
+		{
+			UE_LOG(LogHMD, Warning, TEXT("Failed to initialize core functions. Please check that you have a valid OpenXR runtime installed."));
+			return false;
+		}
+	}
+
 	for (IOpenXRExtensionPlugin* Module : ExtensionPlugins)
 	{
 		Module->PostCreateInstance(Instance);
