OpenXR: fix incorrect warning message (trivial)

#rnx

diff --git a/Engine/Plugins/Runtime/OpenXR/Source/OpenXRHMD/Private/OpenXRHMDModule.cpp b/Engine/Plugins/Runtime/OpenXR/Source/OpenXRHMD/Private/OpenXRHMDModule.cpp
--- a/Engine/Plugins/Runtime/OpenXR/Source/OpenXRHMD/Private/OpenXRHMDModule.cpp
+++ b/Engine/Plugins/Runtime/OpenXR/Source/OpenXRHMD/Private/OpenXRHMDModule.cpp
@@ -774,7 +774,7 @@ bool FOpenXRHMDModule::InitInstance()
 
 		if (StringToPathTest != XR_SUCCESS)
 		{
-			UE_LOG(LogHMD, Warning, TEXT("Failed to initialize core functions. Please check that you have a valid OpenXR runtime installed."));
+			UE_LOG(LogHMD, Warning, TEXT("Instance does not support expected usage of xrStringToPath. Instance is not viable."));
 			return false;
 		}
 	}
