Allow registering external generators in VariableShadingRateImageManager.

As of the moment, only internal generators are used. Allow registering external generators so new ones can be added. Also decoupling FBFoveationImageGenerator from OpenXRHMD to register it as an external image generator.

diff --git a/Engine/Plugins/Runtime/OpenXR/Source/OpenXRHMD/Private/FBFoveationImageGenerator.cpp b/Engine/Plugins/Runtime/OpenXR/Source/OpenXRHMD/Private/FBFoveationImageGenerator.cpp
new file mode 100644
--- /dev/null
+++ b/Engine/Plugins/Runtime/OpenXR/Source/OpenXRHMD/Private/FBFoveationImageGenerator.cpp
@@ -0,0 +1,134 @@
+// Copyright Epic Games, Inc. All Rights Reserved.
+
+#include "FBFoveationImageGenerator.h"
+#include "OpenXRHMD.h"
+#include "OpenXRHMD_Swapchain.h"
+#include "RenderGraphBuilder.h"
+
+static TAutoConsoleVariable<int32> CVarOpenXRFBFoveationLevel(
+	TEXT("xr.OpenXRFBFoveationLevel"),
+	0,
+	TEXT("Possible foveation levels as specified by the XrFoveationLevelFB enumeration.\n")
+	TEXT("0 = None, 1 = Low , 2 = Medium, 3 = High.\n"),
+	ECVF_Default);
+
+static TAutoConsoleVariable<bool> CVarOpenXRFBFoveationDynamic(
+	TEXT("xr.OpenXRFBFoveationDynamic"),
+	false,
+	TEXT("Whether dynamically changing foveation based on performance headroom is enabled.\n"),
+	ECVF_Default);
+
+static TAutoConsoleVariable<float> CVarOpenXRFBFoveationVerticalOffset(
+	TEXT("xr.OpenXRFBFoveationVerticalOffset"),
+	0,
+	TEXT("Desired vertical offset in degrees for the center of the foveation pattern.\n"),
+	ECVF_Default);
+
+static TAutoConsoleVariable<int32> CVarFBFoveationPreview(
+	TEXT("xr.OpenXRFBFoveation.Preview"),
+	1,
+	TEXT("Whether to include FB foveation in VRS preview overlay. This is currently not implemented.")
+	TEXT("0 - off, 1 - on (default)"),
+	ECVF_RenderThreadSafe);
+
+FFBFoveationImageGenerator::FFBFoveationImageGenerator(bool bIsFoveationExtensionSupported, XrInstance InInstance, FOpenXRHMD* HMD, bool bMobileMultiViewEnabled)
+	: bIsMobileMultiViewEnabled(bMobileMultiViewEnabled)
+	, bFoveationExtensionSupported(bIsFoveationExtensionSupported)
+{
+	if (bFoveationExtensionSupported)
+	{
+		XR_ENSURE(xrGetInstanceProcAddr(InInstance, "xrCreateFoveationProfileFB", (PFN_xrVoidFunction*)&xrCreateFoveationProfileFB));
+		XR_ENSURE(xrGetInstanceProcAddr(InInstance, "xrUpdateSwapchainFB", (PFN_xrVoidFunction*)&xrUpdateSwapchainFB));
+		XR_ENSURE(xrGetInstanceProcAddr(InInstance, "xrDestroyFoveationProfileFB", (PFN_xrVoidFunction*)&xrDestroyFoveationProfileFB));
+
+		int32 SanitisedFoveationLevel = FMath::Clamp(CVarOpenXRFBFoveationLevel->GetInt(), XrFoveationLevelFB::XR_FOVEATION_LEVEL_NONE_FB, XrFoveationLevelFB::XR_FOVEATION_LEVEL_HIGH_FB);
+		bool bFoveationDynamic = CVarOpenXRFBFoveationDynamic->GetBool();
+		float SanitisedVerticalOffset = CVarOpenXRFBFoveationVerticalOffset->GetFloat();
+		SanitisedVerticalOffset = SanitisedVerticalOffset >= 0 ? SanitisedVerticalOffset : 0;
+
+		FoveationLevel = static_cast<XrFoveationLevelFB>(SanitisedFoveationLevel);
+		VerticalOffset = SanitisedVerticalOffset;
+		FoveationDynamic = bFoveationDynamic ? XR_FOVEATION_DYNAMIC_LEVEL_ENABLED_FB : XR_FOVEATION_DYNAMIC_DISABLED_FB;
+	}
+
+	OpenXRHMD = HMD;
+}
+
+FRDGTextureRef FFBFoveationImageGenerator::GetImage(FRDGBuilder& GraphBuilder, const FViewInfo& ViewInfo, FVariableRateShadingImageManager::EVRSImageType ImageType)
+{
+	if (!bFoveationExtensionSupported || !OpenXRHMD || !FoveationImages.IsEmpty())
+	{
+		return nullptr;
+	}
+	FOpenXRSwapchain* ColorSwapchain = OpenXRHMD->GetColorSwapchain_RenderThread();
+	if (ColorSwapchain != nullptr)
+	{
+		FTextureRHIRef SwapchainTexture = FoveationImages[ColorSwapchain->GetSwapChainIndex_RHIThread()];
+		TRefCountPtr<IPooledRenderTarget> PooledRenderTarget = CreateRenderTarget(SwapchainTexture, *SwapchainTexture->GetName().ToString());
+
+		return GraphBuilder.RegisterExternalTexture(PooledRenderTarget, *SwapchainTexture->GetName().ToString(), ERDGTextureFlags::SkipTracking);
+	}
+	else
+	{
+		UE_LOG(LogHMD, Error, TEXT("No valid color swapchain to get foveation images from."));
+		return nullptr;
+	}
+}
+
+void FFBFoveationImageGenerator::PrepareImages(FRDGBuilder& GraphBuilder, const FSceneViewFamily& ViewFamily, const FMinimalSceneTextures& SceneTextures)
+{
+	return; //Currently not implemented as the images are prepared only when the color swapchain is reallocated.
+}
+
+bool FFBFoveationImageGenerator::IsEnabledForView(const FSceneView& View) const
+{
+	if (OpenXRHMD)
+	{
+		return OpenXRHMD->IsStereoEnabled() && bFoveationExtensionSupported && View.StereoPass != EStereoscopicPass::eSSP_FULL;
+	}
+	return false;
+}
+
+// This is currently not implemented.
+FRDGTextureRef FFBFoveationImageGenerator::GetDebugImage(FRDGBuilder& GraphBuilder, const FViewInfo& ViewInfo, FVariableRateShadingImageManager::EVRSImageType ImageType)
+{
+	return nullptr;
+}
+
+void FFBFoveationImageGenerator::UpdateFoveationImages()
+{
+	if (!OpenXRHMD)
+	{
+		return;
+	}
+
+	XrFoveationLevelProfileCreateInfoFB FoveationLevelProfileInfo{ XR_TYPE_FOVEATION_LEVEL_PROFILE_CREATE_INFO_FB };
+	FoveationLevelProfileInfo.next = nullptr;
+	FoveationLevelProfileInfo.level = FoveationLevel;
+	FoveationLevelProfileInfo.verticalOffset = VerticalOffset;
+	FoveationLevelProfileInfo.dynamic = FoveationDynamic;
+
+	XrFoveationProfileCreateInfoFB FoveationCreateInfo{ XR_TYPE_FOVEATION_PROFILE_CREATE_INFO_FB };
+	FoveationCreateInfo.next = &FoveationLevelProfileInfo;
+
+	XrFoveationProfileFB FoveationProfile = XR_NULL_HANDLE;
+
+	XrSession Session = OpenXRHMD->GetSession();
+	FOpenXRSwapchain* ColorSwapchain = OpenXRHMD->GetColorSwapchain_RenderThread();
+
+	if (Session && ColorSwapchain != nullptr)
+	{
+		XR_ENSURE(xrCreateFoveationProfileFB(Session, &FoveationCreateInfo, &FoveationProfile));
+
+		XrSwapchainStateFoveationFB SwapchainFoveationState{ XR_TYPE_SWAPCHAIN_STATE_FOVEATION_FB };
+		SwapchainFoveationState.flags = 0; // As per OpenXR specification.
+		SwapchainFoveationState.next = nullptr;
+		SwapchainFoveationState.profile = FoveationProfile;
+
+		XR_ENSURE(xrUpdateSwapchainFB(ColorSwapchain->GetHandle(), reinterpret_cast<const XrSwapchainStateBaseHeaderFB*>(&SwapchainFoveationState)));
+
+		XR_ENSURE(xrDestroyFoveationProfileFB(FoveationProfile));
+
+		ColorSwapchain->GetFragmentDensityMaps(FoveationImages, bIsMobileMultiViewEnabled);
+	}
+}
\ No newline at end of file
diff --git a/Engine/Plugins/Runtime/OpenXR/Source/OpenXRHMD/Private/FBFoveationImageGenerator.h b/Engine/Plugins/Runtime/OpenXR/Source/OpenXRHMD/Private/FBFoveationImageGenerator.h
new file mode 100644
--- /dev/null
+++ b/Engine/Plugins/Runtime/OpenXR/Source/OpenXRHMD/Private/FBFoveationImageGenerator.h
@@ -0,0 +1,47 @@
+// Copyright Epic Games, Inc. All Rights Reserved.
+
+#pragma once
+
+#include "VariableRateShadingImageManager.h"
+
+#include <openxr/openxr.h>
+
+class FRGDBuilder;
+class FOpenXRHMD;
+
+class FFBFoveationImageGenerator : public IVariableRateShadingImageGenerator
+{
+public:
+	FFBFoveationImageGenerator(bool bIsFoveationExtensionSupported, XrInstance InInstance, FOpenXRHMD* HMD, bool bMobileMultiViewEnabled);
+	virtual ~FFBFoveationImageGenerator() override {};
+
+	bool IsFoveationExtensionEnabled() { return bFoveationExtensionSupported;}
+	void UpdateFoveationImages();
+
+	/** IVariableRateShadingImageGenerator interface */
+	virtual FRDGTextureRef GetImage(FRDGBuilder& GraphBuilder, const FViewInfo& ViewInfo, FVariableRateShadingImageManager::EVRSImageType ImageType) override;
+	virtual void PrepareImages(FRDGBuilder& GraphBuilder, const FSceneViewFamily& ViewFamily, const FMinimalSceneTextures& SceneTextures) override;
+	virtual bool IsEnabledForView(const FSceneView& View) const override;
+	virtual FRDGTextureRef GetDebugImage(FRDGBuilder& GraphBuilder, const FViewInfo& ViewInfo, FVariableRateShadingImageManager::EVRSImageType ImageType) override;
+	virtual FVariableRateShadingImageManager::EVRSSourceType GetType() const override
+	{
+		return FVariableRateShadingImageManager::EVRSSourceType::FixedFoveation;
+	}
+	/** End of IVariableRateShadingImageGenerator interface */
+
+private:
+
+	FOpenXRHMD* OpenXRHMD;
+	bool		bIsMobileMultiViewEnabled;
+
+	// XR_FB_foveation
+	bool					bFoveationExtensionSupported;
+	XrFoveationLevelFB		FoveationLevel;
+	float					VerticalOffset;
+	XrFoveationDynamicFB	FoveationDynamic;
+	TArray<FTextureRHIRef>	FoveationImages;
+	PFN_xrCreateFoveationProfileFB	xrCreateFoveationProfileFB;
+	PFN_xrUpdateSwapchainFB			xrUpdateSwapchainFB;
+	PFN_xrDestroyFoveationProfileFB xrDestroyFoveationProfileFB;
+};
+
diff --git a/Engine/Plugins/Runtime/OpenXR/Source/OpenXRHMD/Private/OpenXRHMD.cpp b/Engine/Plugins/Runtime/OpenXR/Source/OpenXRHMD/Private/OpenXRHMD.cpp
--- a/Engine/Plugins/Runtime/OpenXR/Source/OpenXRHMD/Private/OpenXRHMD.cpp
+++ b/Engine/Plugins/Runtime/OpenXR/Source/OpenXRHMD/Private/OpenXRHMD.cpp
@@ -4,6 +4,7 @@
 #include "OpenXRHMD_Layer.h"
 #include "OpenXRHMD_RenderBridge.h"
 #include "OpenXRHMD_Swapchain.h"
+#include "OpenXRHMDSettings.h"
 #include "OpenXRCore.h"
 #include "IOpenXRExtensionPlugin.h"
 #include "IOpenXRHMDModule.h"
@@ -37,6 +38,7 @@
 #include "ScreenRendering.h"
 #include "StereoRenderUtils.h"
 #include "DefaultStereoLayers.h"
+#include "FBFoveationImageGenerator.h"
 
 #if WITH_EDITOR
 #include "Editor.h"
@@ -102,32 +104,6 @@ static TAutoConsoleVariable<bool> CVarOpenXRAllowDepthLayer(
 	TEXT("Enables the depth composition layer if the XR_KHR_composition_layer_depth extension is supported.\n"),
 	ECVF_Default);
 
-static TAutoConsoleVariable<int32> CVarOpenXRFBFoveationLevel(
-	TEXT("xr.OpenXRFBFoveationLevel"),
-	0,
-	TEXT("Possible foveation levels as specified by the XrFoveationLevelFB enumeration.\n")
-	TEXT("0 = None, 1 = Low , 2 = Medium, 3 = High.\n"),
-	ECVF_Default);
-
-static TAutoConsoleVariable<bool> CVarOpenXRFBFoveationDynamic(
-	TEXT("xr.OpenXRFBFoveationDynamic"),
-	false,
-	TEXT("Whether dynamic changing foveation based on performance headroom is enabled.\n"),
-	ECVF_Default);
-
-static TAutoConsoleVariable<float> CVarOpenXRFBFoveationVerticalOffset(
-	TEXT("xr.OpenXRFBFoveationVerticalOffset"),
-	0,
-	TEXT("Desired vertical offset in degrees for the center of the foveation pattern.\n"),
-	ECVF_Default);
-
-static TAutoConsoleVariable<int32> CVarFBFoveationPreview(
-	TEXT("xr.OpenXRFBFoveation.Preview"),
-	1,
-	TEXT("Whether to include FB foveation in VRS preview overlay. This is currently not implemented.")
-	TEXT("0 - off, 1 - on (default)"),
-	ECVF_RenderThreadSafe);
-
 namespace {
 	static TSet<XrViewConfigurationType> SupportedViewConfigurations{ XR_VIEW_CONFIGURATION_TYPE_PRIMARY_STEREO, XR_VIEW_CONFIGURATION_TYPE_PRIMARY_QUAD_VARJO };
 
@@ -1328,9 +1304,6 @@ FOpenXRHMD::FOpenXRHMD(const FAutoRegister& AutoRegister, XrInstance InInstance,
 	XR_ENSURE(xrGetInstanceProperties(Instance, &InstanceProperties));
 	InstanceProperties.runtimeName[XR_MAX_RUNTIME_NAME_SIZE - 1] = 0; // Ensure the name is null terminated.
 
-	SystemHandTrackingProperties = { XR_TYPE_SYSTEM_HAND_TRACKING_PROPERTIES_EXT };
-	SystemProperties = XrSystemProperties{ XR_TYPE_SYSTEM_PROPERTIES, &SystemHandTrackingProperties };
-
 	bDepthExtensionSupported = IsExtensionEnabled(XR_KHR_COMPOSITION_LAYER_DEPTH_EXTENSION_NAME) && CheckPlatformDepthExtensionSupport(InstanceProperties);
 	bHiddenAreaMaskSupported = IsExtensionEnabled(XR_KHR_VISIBILITY_MASK_EXTENSION_NAME) &&
 		!FCStringAnsi::Strstr(InstanceProperties.runtimeName, "Oculus");
@@ -1348,22 +1321,6 @@ FOpenXRHMD::FOpenXRHMD(const FAutoRegister& AutoRegister, XrInstance InInstance,
 	bFoveationExtensionSupported &= IsExtensionEnabled(XR_FB_FOVEATION_VULKAN_EXTENSION_NAME) && GRHISupportsAttachmentVariableRateShading;
 #endif
 
-	if(bFoveationExtensionSupported)
-	{
-		XR_ENSURE(xrGetInstanceProcAddr(Instance, "xrCreateFoveationProfileFB", (PFN_xrVoidFunction*)&xrCreateFoveationProfileFB));
-		XR_ENSURE(xrGetInstanceProcAddr(Instance, "xrUpdateSwapchainFB", (PFN_xrVoidFunction*)&xrUpdateSwapchainFB));
-		XR_ENSURE(xrGetInstanceProcAddr(Instance, "xrDestroyFoveationProfileFB", (PFN_xrVoidFunction*)&xrDestroyFoveationProfileFB));
-
-		int32 SanitisedFoveationLevel = FMath::Clamp(CVarOpenXRFBFoveationLevel->GetInt(), 0, 3);
-		bool bFoveationDynamic = CVarOpenXRFBFoveationDynamic->GetBool();
-		float SanitisedVerticalOffset = CVarOpenXRFBFoveationVerticalOffset->GetFloat();
-		SanitisedVerticalOffset = SanitisedVerticalOffset >= 0 ? SanitisedVerticalOffset : 0;
-
-		FoveationLevel = static_cast<XrFoveationLevelFB>(SanitisedFoveationLevel);
-		VerticalOffset = SanitisedVerticalOffset;
-		FoveationDynamic = bFoveationDynamic ? XR_FOVEATION_DYNAMIC_LEVEL_ENABLED_FB : XR_FOVEATION_DYNAMIC_DISABLED_FB;
-	}
-
 #if PLATFORM_HOLOLENS || PLATFORM_ANDROID
 	bIsStandaloneStereoOnlyDevice = IStereoRendering::IsStartInVR();
 #else
@@ -1726,8 +1683,10 @@ bool FOpenXRHMD::OnStereoStartup()
 	}
 
 	// Retrieve system properties and check for hand tracking support
+	XrSystemHandTrackingPropertiesEXT HandTrackingSystemProperties = { XR_TYPE_SYSTEM_HAND_TRACKING_PROPERTIES_EXT };
+	SystemProperties = XrSystemProperties{ XR_TYPE_SYSTEM_PROPERTIES, &HandTrackingSystemProperties };
 	XR_ENSURE(xrGetSystemProperties(Instance, System, &SystemProperties));
-	bSupportsHandTracking = SystemHandTrackingProperties.supportsHandTracking == XR_TRUE;
+	bSupportsHandTracking = HandTrackingSystemProperties.supportsHandTracking == XR_TRUE;
 
 	// Some runtimes aren't compliant with their number of layers supported.
 	// We support a fallback by emulating non-facelocked layers
@@ -2203,10 +2162,24 @@ void FOpenXRHMD::OnBeginPlay(FWorldContext& InWorldContext)
 {
 	bOpenXRForceStereoLayersEmulationCVarCachedValue = CVarOpenXRForceStereoLayerEmulation.GetValueOnGameThread();
 	bOpenXRInvertAlphaCvarCachedValue = CVarOpenXRInvertAlpha.GetValueOnGameThread();
+
+	const UOpenXRHMDSettings* Settings = GetDefault<UOpenXRHMDSettings>();
+	ensure(Settings);
+	bRuntimeFoveationSupported = bFoveationExtensionSupported && Settings->bIsFBFoveationEnabled;
+	if (bRuntimeFoveationSupported)
+	{
+		FBFoveationImageGenerator = MakeUnique<FFBFoveationImageGenerator>(bRuntimeFoveationSupported, Instance, this, bIsMobileMultiViewEnabled);
+		GVRSImageManager.RegisterExternalImageGenerator(FBFoveationImageGenerator.Get());
+	}
 }
 
 void FOpenXRHMD::OnEndPlay(FWorldContext& InWorldContext)
 {
+	if (bRuntimeFoveationSupported)
+	{
+		GVRSImageManager.UnregisterExternalImageGenerator(FBFoveationImageGenerator.Get());
+		FBFoveationImageGenerator.Reset();
+	}
 }
 
 IStereoRenderTargetManager* FOpenXRHMD::GetRenderTargetManager()
@@ -2266,7 +2239,7 @@ bool FOpenXRHMD::AllocateRenderTargetTextures(uint32 SizeX, uint32 SizeY, uint8
 	}
 	ETextureCreateFlags AuxiliaryCreateFlags = ETextureCreateFlags::None;
 
-	if(bFoveationExtensionSupported)
+	if(FBFoveationImageGenerator && FBFoveationImageGenerator->IsFoveationExtensionEnabled())
 	{
 		AuxiliaryCreateFlags |= TexCreate_Foveation;
 	}
@@ -2301,11 +2274,9 @@ bool FOpenXRHMD::AllocateRenderTargetTextures(uint32 SizeX, uint32 SizeY, uint8
 				Swapchain->IncrementSwapChainIndex_RHIThread();
 			});
 		}
-		if (bFoveationExtensionSupported)
+		if (FBFoveationImageGenerator && FBFoveationImageGenerator->IsFoveationExtensionEnabled())
 		{
-			UpdateFoveationImages();
-			FOpenXRSwapchain* ColorSwapchain = static_cast<FOpenXRSwapchain*>(Swapchain.Get());
-			ColorSwapchain->GetFragmentDensityMaps(FoveationImages, bIsMobileMultiViewEnabled);
+			FBFoveationImageGenerator->UpdateFoveationImages();
 		}
 	}
 
@@ -3950,74 +3921,16 @@ void FOpenXRHMD::UpdateLayer(FOpenXRLayer& ManagerLayer, uint32 LayerId, bool bI
 	});
 }
 
-//---------------------------------------------------
-// IVariableRateShadingImageGenerator interface for XR_FB_foveation Implementation
-//---------------------------------------------------
-
-FRDGTextureRef FOpenXRHMD::GetImage(FRDGBuilder& GraphBuilder, const FViewInfo& ViewInfo, FVariableRateShadingImageManager::EVRSImageType ImageType)
+FOpenXRSwapchain* FOpenXRHMD::GetColorSwapchain_RenderThread()
 {
-	if (!bFoveationExtensionSupported)
+	if (PipelinedLayerStateRendering.ColorSwapchain != nullptr)
 	{
-		return nullptr;
+		return static_cast<FOpenXRSwapchain*>(PipelinedLayerStateRendering.ColorSwapchain.Get());
 	}
-	if (PipelinedLayerStateRendering.ColorSwapchain != nullptr && !FoveationImages.IsEmpty())
-	{
-		FOpenXRSwapchain* ColorSwapchain = static_cast<FOpenXRSwapchain*>(PipelinedLayerStateRendering.ColorSwapchain.Get());
-		
-		FTextureRHIRef SwapchainTexture = FoveationImages[ColorSwapchain->GetSwapChainIndex_RHIThread()];
-		TRefCountPtr<IPooledRenderTarget> PooledRenderTarget = CreateRenderTarget(SwapchainTexture, *SwapchainTexture->GetName().ToString());
 
-		return GraphBuilder.RegisterExternalTexture(PooledRenderTarget, *SwapchainTexture->GetName().ToString(), ERDGTextureFlags::SkipTracking);
-	}
-	else 
-	{
-		UE_LOG(LogHMD, Error, TEXT("No valid color swapchain to get foveation images from."));
-		return nullptr;
-	}
-}
-
-void FOpenXRHMD::PrepareImages(FRDGBuilder& GraphBuilder, const FSceneViewFamily& ViewFamily, const FMinimalSceneTextures& SceneTextures)
-{
-	return; //Currently not implemented as the images are prepared only when the color swapchain is reallocated.
-}
-
-bool FOpenXRHMD::IsEnabledForView(const FSceneView& View) const
-{
-	return IsStereoEnabled() && bFoveationExtensionSupported && View.StereoPass != EStereoscopicPass::eSSP_FULL;
-}
-
-// This is currently not implemented.
-FRDGTextureRef FOpenXRHMD::GetDebugImage(FRDGBuilder& GraphBuilder, const FViewInfo& ViewInfo, FVariableRateShadingImageManager::EVRSImageType ImageType)
-{
 	return nullptr;
 }
 
-void FOpenXRHMD::UpdateFoveationImages()
-{
-	XrFoveationLevelProfileCreateInfoFB FoveationLevelProfileInfo{ XR_TYPE_FOVEATION_LEVEL_PROFILE_CREATE_INFO_FB };
-	FoveationLevelProfileInfo.next = nullptr;
-	FoveationLevelProfileInfo.level = FoveationLevel;
-	FoveationLevelProfileInfo.verticalOffset = VerticalOffset;
-	FoveationLevelProfileInfo.dynamic = FoveationDynamic;
-
-	XrFoveationProfileCreateInfoFB FoveationCreateInfo{ XR_TYPE_FOVEATION_PROFILE_CREATE_INFO_FB };
-	FoveationCreateInfo.next = &FoveationLevelProfileInfo;
-
-	XrFoveationProfileFB FoveationProfile = XR_NULL_HANDLE;
-	XR_ENSURE(xrCreateFoveationProfileFB(Session, &FoveationCreateInfo, &FoveationProfile));
-
-	XrSwapchainStateFoveationFB SwapchainFoveationState{ XR_TYPE_SWAPCHAIN_STATE_FOVEATION_FB };
-	SwapchainFoveationState.flags = 0; // As per OpenXR specification.
-	SwapchainFoveationState.next = nullptr;
-	SwapchainFoveationState.profile = FoveationProfile;
-
-	FOpenXRSwapchain* ColorSwapchain = static_cast<FOpenXRSwapchain*>(PipelinedLayerStateRendering.ColorSwapchain.Get());
-
-	XR_ENSURE(xrUpdateSwapchainFB(ColorSwapchain->GetHandle(), reinterpret_cast<const XrSwapchainStateBaseHeaderFB*>(&SwapchainFoveationState)));
-
-	XR_ENSURE(xrDestroyFoveationProfileFB(FoveationProfile));
-}
-
 //---------------------------------------------------
 // OpenXR Action Space Implementation
 //---------------------------------------------------
diff --git a/Engine/Plugins/Runtime/OpenXR/Source/OpenXRHMD/Private/OpenXRHMD.h b/Engine/Plugins/Runtime/OpenXR/Source/OpenXRHMD/Private/OpenXRHMD.h
--- a/Engine/Plugins/Runtime/OpenXR/Source/OpenXRHMD/Private/OpenXRHMD.h
+++ b/Engine/Plugins/Runtime/OpenXR/Source/OpenXRHMD/Private/OpenXRHMD.h
@@ -16,13 +16,14 @@
 #include "IHeadMountedDisplayVulkanExtensions.h"
 #include "IOpenXRExtensionPluginDelegates.h"
 #include "Misc/EnumClassFlags.h"
-#include "VariableRateShadingImageManager.h"
 
 #include <openxr/openxr.h>
 
 class APlayerController;
 class FSceneView;
 class FSceneViewFamily;
+class FFBFoveationImageGenerator;
+class FOpenXRSwapchain;
 class UCanvas;
 class FOpenXRRenderBridge;
 class IOpenXRInputModule;
@@ -39,7 +40,6 @@ class FOpenXRHMD
 	, public FOpenXRAssetManager
 	, public TStereoLayerManager<FOpenXRLayer>
 	, public IOpenXRExtensionPluginDelegates
-	, public IVariableRateShadingImageGenerator
 {
 private:
 
@@ -371,19 +371,6 @@ public:
 	/** IStereoLayers */
 	virtual bool ShouldCopyDebugLayersToSpectatorScreen() const override { return true; }
 
-	/** IVariableRateShadingImageGenerator interface for XR_FB_foveation */
-	virtual FRDGTextureRef GetImage(FRDGBuilder& GraphBuilder, const FViewInfo& ViewInfo, FVariableRateShadingImageManager::EVRSImageType ImageType) override;
-	virtual void PrepareImages(FRDGBuilder& GraphBuilder, const FSceneViewFamily& ViewFamily, const FMinimalSceneTextures& SceneTextures) override;
-	virtual bool IsEnabledForView(const FSceneView& View) const override;
-	virtual FRDGTextureRef GetDebugImage(FRDGBuilder& GraphBuilder, const FViewInfo& ViewInfo, FVariableRateShadingImageManager::EVRSImageType ImageType) override;
-	virtual FVariableRateShadingImageManager::EVRSSourceType GetType() const override
-	{
-		return FVariableRateShadingImageManager::EVRSSourceType::FixedFoveation;
-	}
-
-	void UpdateFoveationImages();
-
-
 	/** IOpenXRExtensionPluginDelegates */
 public:
 	virtual FApplyHapticFeedbackAddChainStructsDelegate& GetApplyHapticFeedbackAddChainStructsDelegate() override { return ApplyHapticFeedbackAddChainStructsDelegate; }
@@ -427,7 +414,7 @@ public:
 
 	/** Returns shader platform the plugin is currently configured for, in the editor it can change due to preview platforms. */
 	EShaderPlatform GetConfiguredShaderPlatform() const { check(ConfiguredShaderPlatform != EShaderPlatform::SP_NumPlatforms); return ConfiguredShaderPlatform; }
-
+	FOpenXRSwapchain* GetColorSwapchain_RenderThread();
 private:
 
 	TArray<XrEnvironmentBlendMode> RetrieveEnvironmentBlendModes() const;
@@ -477,7 +464,6 @@ private:
 	XrViewConfigurationType SelectedViewConfigurationType;
 	XrEnvironmentBlendMode  SelectedEnvironmentBlendMode;
 	XrInstanceProperties    InstanceProperties;
-	XrSystemHandTrackingPropertiesEXT SystemHandTrackingProperties;
 	XrSystemProperties      SystemProperties;
 
 	FPipelinedFrameState	PipelinedFrameStateGame;
@@ -512,15 +498,9 @@ private:
 	TArray<IStereoLayers::FLayerDesc> EmulatedFaceLockedLayers;
 	TArray<FOpenXRLayer>			  NativeLayers;
 
-	// XR_FB_foveation
+	TUniquePtr<FFBFoveationImageGenerator> FBFoveationImageGenerator;
 	bool					bFoveationExtensionSupported;
-	XrFoveationLevelFB		FoveationLevel;
-	float					VerticalOffset;
-	XrFoveationDynamicFB	FoveationDynamic;
-	TArray<FTextureRHIRef>	FoveationImages;
-	PFN_xrCreateFoveationProfileFB	xrCreateFoveationProfileFB;
-	PFN_xrUpdateSwapchainFB			xrUpdateSwapchainFB;
-	PFN_xrDestroyFoveationProfileFB xrDestroyFoveationProfileFB;
+	bool					bRuntimeFoveationSupported;
 };
 
 ENUM_CLASS_FLAGS(FOpenXRHMD::EOpenXRLayerStateFlags);
\ No newline at end of file
diff --git a/Engine/Source/Runtime/Renderer/Private/VariableRateShading/VariableRateShadingImageManager.cpp b/Engine/Source/Runtime/Renderer/Private/VariableRateShading/VariableRateShadingImageManager.cpp
--- a/Engine/Source/Runtime/Renderer/Private/VariableRateShading/VariableRateShadingImageManager.cpp
+++ b/Engine/Source/Runtime/Renderer/Private/VariableRateShading/VariableRateShadingImageManager.cpp
@@ -279,10 +279,14 @@ IMPLEMENT_GLOBAL_SHADER(FDebugVariableRateShadingPS, "/Engine/Private/VariableRa
 FVariableRateShadingImageManager::FVariableRateShadingImageManager()
 	: FRenderResource()
 {
-	ImageGenerators.Add(MakeUnique<FFoveatedImageGenerator>());
-	ImageGenerators.Add(MakeUnique<FContrastAdaptiveImageGenerator>());
+	InternalGenerators.Add(MakeUnique<FFoveatedImageGenerator>());
+	InternalGenerators.Add(MakeUnique<FContrastAdaptiveImageGenerator>());
 
-	// TODO: Add more generators or allow registration from plugins
+	FWriteScopeLock GeneratorsLock(GeneratorsMutex);
+	for (const TUniquePtr<IVariableRateShadingImageGenerator>& Generator : InternalGenerators)
+	{
+		ImageGenerators.Add(Generator.Get());
+	}
 }
 
 FVariableRateShadingImageManager::~FVariableRateShadingImageManager() {}
@@ -385,17 +389,20 @@ FRDGTextureRef FVariableRateShadingImageManager::GetVariableRateShadingImage(FRD
 	// Otherwise collate all internal sources
 	TArray<FRDGTextureRef> InternalVRSSources;
 
-	for (TUniquePtr<IVariableRateShadingImageGenerator>& Generator : ImageGenerators)
 	{
-		FRDGTextureRef Image = nullptr;
-		if (Generator->IsEnabledForView(ViewInfo) && !EnumHasAnyFlags(VRSTypesToExclude, Generator->GetType()))
+		FReadScopeLock GeneratorsLock(GeneratorsMutex);	
+		for (IVariableRateShadingImageGenerator* const Generator : ImageGenerators)
 		{
-			Image = Generator->GetImage(GraphBuilder, ViewInfo, ImageType);
-		}
+			FRDGTextureRef Image = nullptr;
+			if (Generator->IsEnabledForView(ViewInfo) && !EnumHasAnyFlags(VRSTypesToExclude, Generator->GetType()))
+			{
+				Image = Generator->GetImage(GraphBuilder, ViewInfo, ImageType);
+			}
 
-		if (Image)
-		{
-			InternalVRSSources.Add(Image);
+			if (Image)
+			{
+				InternalVRSSources.Add(Image);
+			}
 		}
 	}
 
@@ -452,18 +459,22 @@ void FVariableRateShadingImageManager::PrepareImageBasedVRS(FRDGBuilder& GraphBu
 	}
 
 	// Invoke image generators
-	for (TUniquePtr<IVariableRateShadingImageGenerator>& Generator : ImageGenerators)
 	{
-		if (Generator->IsEnabledForView(*ViewFamily.Views[0]))
+		FReadScopeLock GeneratorsLock(GeneratorsMutex);
+		for (IVariableRateShadingImageGenerator* const Generator : ImageGenerators)
 		{
-			Generator->PrepareImages(GraphBuilder, ViewFamily, SceneTextures);
+			if (Generator->IsEnabledForView(*ViewFamily.Views[0]))
+			{
+				Generator->PrepareImages(GraphBuilder, ViewFamily, SceneTextures);
+			}
 		}
 	}
 }
 
-bool FVariableRateShadingImageManager::IsTypeEnabledForView(const FSceneView& View, FVariableRateShadingImageManager::EVRSSourceType Type) const
+bool FVariableRateShadingImageManager::IsTypeEnabledForView(const FSceneView& View, FVariableRateShadingImageManager::EVRSSourceType Type)
 {
-	for (const TUniquePtr<IVariableRateShadingImageGenerator>& Generator : ImageGenerators)
+	FReadScopeLock GeneratorsLock(GeneratorsMutex);
+	for (IVariableRateShadingImageGenerator* const Generator : ImageGenerators)
 	{
 		if (EnumHasAnyFlags(Type, Generator->GetType()) && Generator->IsEnabledForView(View))
 		{
@@ -506,17 +517,20 @@ void FVariableRateShadingImageManager::DrawDebugPreview(FRDGBuilder& GraphBuilde
 			{
 				TArray<FRDGTextureRef> InternalVRSSources;
 
-				for (TUniquePtr<IVariableRateShadingImageGenerator>& Generator : ImageGenerators)
 				{
-					FRDGTextureRef Image = nullptr;
-					if (Generator->IsEnabledForView(*View))
-					{
-						Image = Generator->GetDebugImage(GraphBuilder, *ViewInfo, PreviewImageType);
-					}
-
-					if (Image)
+					FReadScopeLock GeneratorsLock(GeneratorsMutex);
+					for (IVariableRateShadingImageGenerator* const Generator : ImageGenerators)
 					{
-						InternalVRSSources.Add(Image);
+						FRDGTextureRef Image = nullptr;
+						if (Generator->IsEnabledForView(*View))
+						{
+							Image = Generator->GetDebugImage(GraphBuilder, *ViewInfo, PreviewImageType);
+						}
+
+						if (Image)
+						{
+							InternalVRSSources.Add(Image);
+						}
 					}
 				}
 
@@ -576,6 +590,27 @@ void FVariableRateShadingImageManager::DrawDebugPreview(FRDGBuilder& GraphBuilde
 
 }
 
+void FVariableRateShadingImageManager::RegisterExternalImageGenerator(IVariableRateShadingImageGenerator* ExternalGenerator)
+{
+	if (ExternalGenerator == nullptr)
+	{
+		UE_LOG(LogVRS, Warning, TEXT("Trying to register a null VRS generator. Generator will be ignored."));
+		return;
+	}
+	FWriteScopeLock GeneratorsLock(GeneratorsMutex);
+	ImageGenerators.Add(ExternalGenerator);
+}
+
+void FVariableRateShadingImageManager::UnregisterExternalImageGenerator(IVariableRateShadingImageGenerator* ExternalGenerator)
+{
+	if (ExternalGenerator == nullptr)
+	{
+		UE_LOG(LogVRS, Warning, TEXT("Trying to unregister a null VRS generator. Generator will be ignored."));
+		return;
+	}
+	FWriteScopeLock GeneratorsLock(GeneratorsMutex);
+	ImageGenerators.Remove(ExternalGenerator);
+}
 
 /**
  * Private functions
diff --git a/Engine/Source/Runtime/Renderer/Public/VariableRateShadingImageManager.h b/Engine/Source/Runtime/Renderer/Public/VariableRateShadingImageManager.h
--- a/Engine/Source/Runtime/Renderer/Public/VariableRateShadingImageManager.h
+++ b/Engine/Source/Runtime/Renderer/Public/VariableRateShadingImageManager.h
@@ -96,7 +96,10 @@ public:
 	/**
 	 * Returns true if any generator among the given types is enabled, false otherwise.
 	 */
-	bool IsTypeEnabledForView(const FSceneView& View, EVRSSourceType Type) const;
+	bool IsTypeEnabledForView(const FSceneView& View, EVRSSourceType Type);
+
+	RENDERER_API void RegisterExternalImageGenerator(IVariableRateShadingImageGenerator* ExternalGenerator);
+	RENDERER_API void UnregisterExternalImageGenerator(IVariableRateShadingImageGenerator* ExternalGenerator);
 
 	static bool IsVRSSupportedByRHI();
 	static bool IsVRSEnabled();
@@ -111,7 +114,13 @@ public:
 
 private:
 	TRefCountPtr<IPooledRenderTarget> MobileHMDFixedFoveationOverrideImage;
-	TArray<TUniquePtr<IVariableRateShadingImageGenerator>> ImageGenerators;
+	// Used to access the generators functionalities.
+	TArray<IVariableRateShadingImageGenerator*> ImageGenerators;
+	// This is used only to own the memory of generators created in the FVariableRateShadingImageManager constructor.
+	TArray<TUniquePtr<IVariableRateShadingImageGenerator>> InternalGenerators;
+	// Guards ImageGenerators because most of the calls of the manager are on the render thread but the 
+	// Register/UnregisterExternalImageGenerator could come from the game thread.
+	FRWLock	GeneratorsMutex;
 
 	FRDGTextureRef CombineShadingRateImages(FRDGBuilder& GraphBuilder, const FViewInfo& ViewInfo, TArray<FRDGTextureRef> Sources);
 	FRDGTextureRef GetForceRateImage(FRDGBuilder& GraphBuilder, int RateIndex = 0, EVRSImageType ImageType = EVRSImageType::Full);
