Fix compile error due to missing SceneView.

diff --git a/Engine/Plugins/Runtime/OpenXR/Source/OpenXRHMD/Private/FBFoveationImageGenerator.cpp b/Engine/Plugins/Runtime/OpenXR/Source/OpenXRHMD/Private/FBFoveationImageGenerator.cpp
--- a/Engine/Plugins/Runtime/OpenXR/Source/OpenXRHMD/Private/FBFoveationImageGenerator.cpp
+++ b/Engine/Plugins/Runtime/OpenXR/Source/OpenXRHMD/Private/FBFoveationImageGenerator.cpp
@@ -1,6 +1,7 @@
 // Copyright Epic Games, Inc. All Rights Reserved.
 
 #include "FBFoveationImageGenerator.h"
+#include "SceneView.h"
 #include "OpenXRHMD.h"
 #include "OpenXRHMD_Swapchain.h"
 #include "RenderGraphBuilder.h"
diff --git a/Engine/Plugins/Runtime/OpenXR/Source/OpenXRHMD/Private/FBFoveationImageGenerator.h b/Engine/Plugins/Runtime/OpenXR/Source/OpenXRHMD/Private/FBFoveationImageGenerator.h
--- a/Engine/Plugins/Runtime/OpenXR/Source/OpenXRHMD/Private/FBFoveationImageGenerator.h
+++ b/Engine/Plugins/Runtime/OpenXR/Source/OpenXRHMD/Private/FBFoveationImageGenerator.h
@@ -8,6 +8,7 @@
 
 class FRGDBuilder;
 class FOpenXRHMD;
+class FSceneView;
 
 class FFBFoveationImageGenerator : public IVariableRateShadingImageGenerator
 {
