Fix to deferencing null pointer.

diff --git a/Engine/Plugins/Runtime/OpenXR/Source/OpenXRHMD/Private/OpenXRHMD.cpp b/Engine/Plugins/Runtime/OpenXR/Source/OpenXRHMD/Private/OpenXRHMD.cpp
--- a/Engine/Plugins/Runtime/OpenXR/Source/OpenXRHMD/Private/OpenXRHMD.cpp
+++ b/Engine/Plugins/Runtime/OpenXR/Source/OpenXRHMD/Private/OpenXRHMD.cpp
@@ -2168,8 +2168,7 @@ void FOpenXRHMD::OnBeginPlay(FWorldContext& InWorldContext)
 	bOpenXRInvertAlphaCvarCachedValue = CVarOpenXRInvertAlpha.GetValueOnGameThread();
 
 	const UOpenXRHMDSettings* Settings = GetDefault<UOpenXRHMDSettings>();
-	ensure(Settings);
-	bRuntimeFoveationSupported = bFoveationExtensionSupported && Settings->bIsFBFoveationEnabled;
+	bRuntimeFoveationSupported = bFoveationExtensionSupported && (Settings != nullptr ? Settings->bIsFBFoveationEnabled : false);
 	if (bRuntimeFoveationSupported)
 	{
 		FBFoveationImageGenerator = MakeUnique<FFBFoveationImageGenerator>(bRuntimeFoveationSupported, Instance, this, bIsMobileMultiViewEnabled);
