Make Vulkan RHI create textures with foveation image layout if TexCreate_Foveation is set.

diff --git a/Engine/Plugins/Runtime/OpenXR/Source/OpenXRHMD/Private/OpenXRHMD_Swapchain.cpp b/Engine/Plugins/Runtime/OpenXR/Source/OpenXRHMD/Private/OpenXRHMD_Swapchain.cpp
--- a/Engine/Plugins/Runtime/OpenXR/Source/OpenXRHMD/Private/OpenXRHMD_Swapchain.cpp
+++ b/Engine/Plugins/Runtime/OpenXR/Source/OpenXRHMD/Private/OpenXRHMD_Swapchain.cpp
@@ -261,14 +261,11 @@ void FOpenXRSwapchain::GetFragmentDensityMaps(TArray<FTextureRHIRef>& OutTexture
 	XrSwapchain Swapchain = GetHandle();
 
 	TArray<XrSwapchainImageVulkanKHR> Images = EnumerateImages<XrSwapchainImageVulkanKHR>(Swapchain, XR_TYPE_SWAPCHAIN_IMAGE_VULKAN_KHR, Next);
-	const VkImageSubresourceRange SubresourceRangeAll = { VK_IMAGE_ASPECT_COLOR_BIT, 0, VK_REMAINING_MIP_LEVELS, 0, VK_REMAINING_ARRAY_LAYERS };
 
 	for (const XrSwapchainImageVulkanKHR& Image : Images)
 	{
 		const XrBaseOutStructure* NextHeader = reinterpret_cast<const XrBaseOutStructure*>(Image.next);
 		const XrSwapchainImageFoveationVulkanFB* FoveationImageVulkan = reinterpret_cast<const XrSwapchainImageFoveationVulkanFB*>(NextHeader);
-		// TODO: remove RHISetImageLayout when infrastructure to set an initial Vulkan layout is added.
-		VulkanRHI->RHISetImageLayout((VkImage)FoveationImageVulkan->image, VK_IMAGE_LAYOUT_UNDEFINED, VK_IMAGE_LAYOUT_FRAGMENT_DENSITY_MAP_OPTIMAL_EXT, SubresourceRangeAll);
 		OutTextureChain.Add(static_cast<FTextureRHIRef>(bIsMobileMultiViewEnabled ?
 			VulkanRHI->RHICreateTexture2DArrayFromResource(GPixelFormats[PF_R8G8].UnrealFormat, FoveationImageVulkan->width, FoveationImageVulkan->height, 2, 1, 1, FoveationImageVulkan->image, TexCreate_Foveation, FClearValueBinding::White) :
 			VulkanRHI->RHICreateTexture2DFromResource(GPixelFormats[PF_R8G8].UnrealFormat, FoveationImageVulkan->width, FoveationImageVulkan->height, 1, 1, FoveationImageVulkan->image, TexCreate_Foveation, FClearValueBinding::White)
diff --git a/Engine/Source/Runtime/RHI/Private/RHIUtilities.cpp b/Engine/Source/Runtime/RHI/Private/RHIUtilities.cpp
--- a/Engine/Source/Runtime/RHI/Private/RHIUtilities.cpp
+++ b/Engine/Source/Runtime/RHI/Private/RHIUtilities.cpp
@@ -610,6 +610,10 @@ ERHIAccess RHIGetDefaultResourceState(ETextureCreateFlags InUsage, bool bInHasIn
 		{
 			ResourceState = ERHIAccess::SRVMask;
 		}
+		else if (EnumHasAnyFlags(InUsage, TexCreate_Foveation))
+		{
+			ResourceState = ERHIAccess::ShadingRateSource;
+		}
 	}
 
 	return ResourceState;
diff --git a/Engine/Source/Runtime/VulkanRHI/Private/VulkanTexture.cpp b/Engine/Source/Runtime/VulkanRHI/Private/VulkanTexture.cpp
--- a/Engine/Source/Runtime/VulkanRHI/Private/VulkanTexture.cpp
+++ b/Engine/Source/Runtime/VulkanRHI/Private/VulkanTexture.cpp
@@ -546,6 +546,7 @@ static VkImageLayout GetInitialLayoutFromRHIAccess(ERHIAccess RHIAccess, bool bI
 		case ERHIAccess::Discard:	return VK_IMAGE_LAYOUT_UNDEFINED;
 		case ERHIAccess::CopySrc:	return VK_IMAGE_LAYOUT_TRANSFER_SRC_OPTIMAL;
 		case ERHIAccess::CopyDest:	return VK_IMAGE_LAYOUT_TRANSFER_DST_OPTIMAL;
+		case ERHIAccess::ShadingRateSource:	return VK_IMAGE_LAYOUT_FRAGMENT_DENSITY_MAP_OPTIMAL_EXT; // TODO: This should be set to FSR or FDM depending on the active extension. Default to FDM for the moment.
 	}
 
 	checkf(false, TEXT("Invalid initial access %d"), RHIAccess);
