Fix foveation turning on when FSR attachment is available.

GRHISupportsAttachmentVariableRateShading could be enabled by FSR even when FDM is not available. Need to check whether the active extension is FDM by checking the ImageDataType.

diff --git a/Engine/Plugins/Runtime/OpenXR/Source/OpenXRHMD/Private/OpenXRHMD.cpp b/Engine/Plugins/Runtime/OpenXR/Source/OpenXRHMD/Private/OpenXRHMD.cpp
--- a/Engine/Plugins/Runtime/OpenXR/Source/OpenXRHMD/Private/OpenXRHMD.cpp
+++ b/Engine/Plugins/Runtime/OpenXR/Source/OpenXRHMD/Private/OpenXRHMD.cpp
@@ -1322,7 +1322,7 @@ FOpenXRHMD::FOpenXRHMD(const FAutoRegister& AutoRegister, XrInstance InInstance,
 		IsExtensionEnabled(XR_FB_FOVEATION_CONFIGURATION_EXTENSION_NAME);
 
 #ifdef XR_USE_GRAPHICS_API_VULKAN
-	bFoveationExtensionSupported &= IsExtensionEnabled(XR_FB_FOVEATION_VULKAN_EXTENSION_NAME) && GRHISupportsAttachmentVariableRateShading;
+	bFoveationExtensionSupported &= IsExtensionEnabled(XR_FB_FOVEATION_VULKAN_EXTENSION_NAME) && GRHISupportsAttachmentVariableRateShading && GRHIVariableRateShadingImageDataType == VRSImage_Fractional;
 #endif
 
 #if PLATFORM_HOLOLENS || PLATFORM_ANDROID
