Fix foveation image access race condition between render and rhi threads.

diff --git a/Engine/Plugins/Runtime/OpenXR/Source/OpenXRHMD/Private/FBFoveationImageGenerator.cpp b/Engine/Plugins/Runtime/OpenXR/Source/OpenXRHMD/Private/FBFoveationImageGenerator.cpp
--- a/Engine/Plugins/Runtime/OpenXR/Source/OpenXRHMD/Private/FBFoveationImageGenerator.cpp
+++ b/Engine/Plugins/Runtime/OpenXR/Source/OpenXRHMD/Private/FBFoveationImageGenerator.cpp
@@ -61,10 +61,9 @@ FRDGTextureRef FFBFoveationImageGenerator::GetImage(FRDGBuilder& GraphBuilder, c
 	{
 		return nullptr;
 	}
-	FOpenXRSwapchain* ColorSwapchain = OpenXRHMD->GetColorSwapchain_RenderThread();
-	if (ColorSwapchain != nullptr)
+	if (FoveationImages.IsValidIndex(CurrentFrameSwapchainIndex))
 	{
-		FTextureRHIRef SwapchainTexture = FoveationImages[ColorSwapchain->GetSwapChainIndex_RHIThread()];
+		FTextureRHIRef SwapchainTexture = FoveationImages[CurrentFrameSwapchainIndex];
 		TRefCountPtr<IPooledRenderTarget> PooledRenderTarget = CreateRenderTarget(SwapchainTexture, *SwapchainTexture->GetName().ToString());
 
 		return GraphBuilder.RegisterExternalTexture(PooledRenderTarget, *SwapchainTexture->GetName().ToString(), ERDGTextureFlags::SkipTracking);
@@ -132,4 +131,9 @@ void FFBFoveationImageGenerator::UpdateFoveationImages()
 
 		ColorSwapchain->GetFragmentDensityMaps(FoveationImages, bIsMobileMultiViewEnabled);
 	}
+}
+
+void FFBFoveationImageGenerator::SetCurrentFrameSwapchainIndex(int32 InCurrentFrameSwapchainIndex)
+{
+	CurrentFrameSwapchainIndex = InCurrentFrameSwapchainIndex;
 }
\ No newline at end of file
diff --git a/Engine/Plugins/Runtime/OpenXR/Source/OpenXRHMD/Private/FBFoveationImageGenerator.h b/Engine/Plugins/Runtime/OpenXR/Source/OpenXRHMD/Private/FBFoveationImageGenerator.h
--- a/Engine/Plugins/Runtime/OpenXR/Source/OpenXRHMD/Private/FBFoveationImageGenerator.h
+++ b/Engine/Plugins/Runtime/OpenXR/Source/OpenXRHMD/Private/FBFoveationImageGenerator.h
@@ -18,6 +18,7 @@ public:
 
 	bool IsFoveationExtensionEnabled() { return bFoveationExtensionSupported;}
 	void UpdateFoveationImages();
+	void SetCurrentFrameSwapchainIndex(int32 CurrentFrameSwapchainIndex);
 
 	/** IVariableRateShadingImageGenerator interface */
 	virtual FRDGTextureRef GetImage(FRDGBuilder& GraphBuilder, const FViewInfo& ViewInfo, FVariableRateShadingImageManager::EVRSImageType ImageType) override;
@@ -34,6 +35,7 @@ private:
 
 	FOpenXRHMD* OpenXRHMD;
 	bool		bIsMobileMultiViewEnabled;
+	int32		CurrentFrameSwapchainIndex;
 
 	// XR_FB_foveation
 	bool					bFoveationExtensionSupported;
diff --git a/Engine/Plugins/Runtime/OpenXR/Source/OpenXRHMD/Private/OpenXRHMD.cpp b/Engine/Plugins/Runtime/OpenXR/Source/OpenXRHMD/Private/OpenXRHMD.cpp
--- a/Engine/Plugins/Runtime/OpenXR/Source/OpenXRHMD/Private/OpenXRHMD.cpp
+++ b/Engine/Plugins/Runtime/OpenXR/Source/OpenXRHMD/Private/OpenXRHMD.cpp
@@ -2828,6 +2828,11 @@ void FOpenXRHMD::OnBeginRendering_RenderThread(FRHICommandListImmediate& RHICmdL
 		// Always check for sanity before using it.
 		FXRSwapChainPtr EmulationSwapchain = PipelinedLayerStateRendering.EmulatedLayerState.EmulationSwapchain;
 
+		if (bFoveationExtensionSupported && FBFoveationImageGenerator.IsValid())
+		{
+			FBFoveationImageGenerator->SetCurrentFrameSwapchainIndex(ColorSwapchain->GetSwapChainIndex_RHIThread());
+		}
+
 		RHICmdList.EnqueueLambda([this, FrameState = PipelinedFrameStateRendering, ColorSwapchain, DepthSwapchain, EmulationSwapchain](FRHICommandListImmediate& InRHICmdList)
 		{
 			OnBeginRendering_RHIThread(FrameState, ColorSwapchain, DepthSwapchain, EmulationSwapchain);
