OpenXR: Attach compositor submissions thread to the Java Env

Logcat spam lead to this discovery. Attaching the current thread to the Java Env results in 250-300us speed increase inside xrEndFrame.

diff --git a/Engine/Plugins/Runtime/OpenXR/Source/OpenXRHMD/Private/OpenXRHMD.cpp b/Engine/Plugins/Runtime/OpenXR/Source/OpenXRHMD/Private/OpenXRHMD.cpp
--- a/Engine/Plugins/Runtime/OpenXR/Source/OpenXRHMD/Private/OpenXRHMD.cpp
+++ b/Engine/Plugins/Runtime/OpenXR/Source/OpenXRHMD/Private/OpenXRHMD.cpp
@@ -40,6 +40,10 @@
 #include "DefaultStereoLayers.h"
 #include "FBFoveationImageGenerator.h"
 
+#if PLATFORM_ANDROID
+#include "Android/AndroidApplication.h"
+#endif
+
 #if WITH_EDITOR
 #include "Editor.h"
 #include "Editor/EditorEngine.h"
@@ -3388,6 +3392,12 @@ void FOpenXRHMD::OnFinishRendering_RHIThread()
 
 		UE_LOG(LogHMD, VeryVerbose, TEXT("Presenting frame predicted to be displayed at %lld"), PipelinedFrameStateRHI.FrameState.predictedDisplayTime);
 
+#if PLATFORM_ANDROID
+		// Android OpenXR runtimes frequently need access to the JNIEnv, so we need to attach the submitting
+		// thread. We have to do this per-frame because we can detach if app loses focus.
+		FAndroidApplication::GetJavaEnv();
+#endif
+
 		XR_ENSURE(xrEndFrame(Session, &EndInfo));
 	}
 
