OpenXR: don't use XrTime 0 in GetPoseForTime

if we call GetPoseForTime after the session has ended, we might have a zero'd time from retained XrFrameState.
In that case, we shouldn't be fetching any poses, nor pretending that poses might be valid.

diff --git a/Engine/Plugins/Runtime/OpenXR/Source/OpenXRHMD/Private/OpenXRHMD.cpp b/Engine/Plugins/Runtime/OpenXR/Source/OpenXRHMD/Private/OpenXRHMD.cpp
--- a/Engine/Plugins/Runtime/OpenXR/Source/OpenXRHMD/Private/OpenXRHMD.cpp
+++ b/Engine/Plugins/Runtime/OpenXR/Source/OpenXRHMD/Private/OpenXRHMD.cpp
@@ -654,6 +654,15 @@ bool FOpenXRHMD::GetPoseForTime(int32 DeviceId, FTimespan Timespan, bool& OutTim
 	{
 		OutTimeWasUsed = false;
 		TargetTime = GetDisplayTime();
+
+		
+		if (TargetTime == 0)
+		{
+			// We might still get an out-of-sync query after the session has ended.
+			// We could return the last known location via PipelineState.DeviceLocations
+			// but UpdateDeviceLocations doesn't do that right now. We'll just fail for now.
+			return false;
+		}
 	}
 	else
 	{
diff --git a/Engine/Plugins/Runtime/OpenXR/Source/OpenXRInput/Private/OpenXRInput.cpp b/Engine/Plugins/Runtime/OpenXR/Source/OpenXRInput/Private/OpenXRInput.cpp
--- a/Engine/Plugins/Runtime/OpenXR/Source/OpenXRInput/Private/OpenXRInput.cpp
+++ b/Engine/Plugins/Runtime/OpenXR/Source/OpenXRInput/Private/OpenXRInput.cpp
@@ -1232,15 +1232,15 @@ bool FOpenXRInputPlugin::FOpenXRInput::GetControllerOrientationAndPositionForTim
 	State.type = XR_TYPE_ACTION_STATE_POSE;
 	State.next = nullptr;
 	XrResult Result = xrGetActionStatePose(Session, &GetInfo, &State);
+	bool bPoseFetched = false;
 	if (Result >= XR_SUCCESS && State.isActive)
 	{
 		FQuat Orientation;
-		OpenXRHMD->GetPoseForTime(GetDeviceIDForMotionSource(MotionSource), Time, OutTimeWasUsed, Orientation, OutPosition, OutbProvidedLinearVelocity, OutLinearVelocity, OutbProvidedAngularVelocity, OutAngularVelocityAsAxisAndLength, OutbProvidedLinearAcceleration, OutLinearAcceleration, WorldToMetersScale);
+		bPoseFetched = OpenXRHMD->GetPoseForTime(GetDeviceIDForMotionSource(MotionSource), Time, OutTimeWasUsed, Orientation, OutPosition, OutbProvidedLinearVelocity, OutLinearVelocity, OutbProvidedAngularVelocity, OutAngularVelocityAsAxisAndLength, OutbProvidedLinearAcceleration, OutLinearAcceleration, WorldToMetersScale);
 		OutOrientation = FRotator(Orientation);
-		return true;
 	}
 
-	return false;
+	return bPoseFetched;
 }
 
 ETrackingStatus FOpenXRInputPlugin::FOpenXRInput::GetControllerTrackingStatus(const int32 ControllerIndex, const FName MotionSource) const
