Updated to support Arm64EC

diff --git a/Engine/Plugins/Runtime/OpenXR/Source/OpenXRHMD/Private/OpenXRHMDModule.cpp b/Engine/Plugins/Runtime/OpenXR/Source/OpenXRHMD/Private/OpenXRHMDModule.cpp
--- a/Engine/Plugins/Runtime/OpenXR/Source/OpenXRHMD/Private/OpenXRHMDModule.cpp
+++ b/Engine/Plugins/Runtime/OpenXR/Source/OpenXRHMD/Private/OpenXRHMDModule.cpp
@@ -360,7 +360,7 @@ bool FOpenXRHMDModule::InitRenderBridge()
 PFN_xrGetInstanceProcAddr FOpenXRHMDModule::GetDefaultLoader()
 {
 #if PLATFORM_WINDOWS
-#if !PLATFORM_CPU_X86_FAMILY
+#if !PLATFORM_CPU_X86_FAMILY && !defined(_M_ARM64EC)
 #error Windows platform does not currently support this CPU family. A OpenXR loader binary for this CPU family is needed.
 #endif
 
