Gate accesses to PipelinedFrameState.Views XrViews array to avoid crashing after XR_ERROR_SESSION_LOST failures.
The app will still exit after detecting XR_SESSION_STATE_LOSS_PENDING, but will no longer crash.

diff --git a/Engine/Plugins/Runtime/OpenXR/Source/OpenXRHMD/Private/OpenXRHMD.cpp b/Engine/Plugins/Runtime/OpenXR/Source/OpenXRHMD/Private/OpenXRHMD.cpp
--- a/Engine/Plugins/Runtime/OpenXR/Source/OpenXRHMD/Private/OpenXRHMD.cpp
+++ b/Engine/Plugins/Runtime/OpenXR/Source/OpenXRHMD/Private/OpenXRHMD.cpp
@@ -2784,33 +2784,36 @@ void FOpenXRHMD::OnBeginRendering_RenderThread(FRHICommandListImmediate& RHICmdL
 
 	const float WorldToMeters = GetWorldToMetersScale();
 
-	for (int32 ViewIndex = 0; ViewIndex < ViewFamily.Views.Num(); ViewIndex++)
+	if (PipelinedFrameStateRendering.Views.Num() == ViewFamily.Views.Num())
 	{
-		if (ViewFamily.Views[ViewIndex]->StereoPass == EStereoscopicPass::eSSP_FULL)
+		for (int32 ViewIndex = 0; ViewIndex < ViewFamily.Views.Num(); ViewIndex++)
 		{
-			continue;
-		}
+			if (ViewFamily.Views[ViewIndex]->StereoPass == EStereoscopicPass::eSSP_FULL)
+			{
+				continue;
+			}
 
-		const XrView& View = PipelinedFrameStateRendering.Views[ViewIndex];
-		FTransform EyePose = ToFTransform(View.pose, WorldToMeters);
+			const XrView& View = PipelinedFrameStateRendering.Views[ViewIndex];
+			FTransform EyePose = ToFTransform(View.pose, WorldToMeters);
 
-		// Apply the base HMD pose to each eye pose, we will late update this pose for late update in another callback
-		FTransform BasePose(ViewFamily.Views[ViewIndex]->BaseHmdOrientation, ViewFamily.Views[ViewIndex]->BaseHmdLocation);
-		FTransform BasePoseTransform = EyePose * BasePose;
-		BasePoseTransform.NormalizeRotation();
-		
-		XrCompositionLayerProjectionView& Projection = PipelinedLayerStateRendering.ProjectionLayers[ViewIndex];
-		Projection.pose = ToXrPose(BasePoseTransform, WorldToMeters);
-		Projection.fov = View.fov;
+			// Apply the base HMD pose to each eye pose, we will late update this pose for late update in another callback
+			FTransform BasePose(ViewFamily.Views[ViewIndex]->BaseHmdOrientation, ViewFamily.Views[ViewIndex]->BaseHmdLocation);
+			FTransform BasePoseTransform = EyePose * BasePose;
+			BasePoseTransform.NormalizeRotation();
 
-		if (EnumHasAnyFlags(PipelinedLayerStateRendering.LayerStateFlags, EOpenXRLayerStateFlags::SubmitEmulatedFaceLockedLayer))
-		{
-			XrCompositionLayerProjectionView& CompositedProjection = PipelinedLayerStateRendering.EmulatedLayerState.CompositedProjectionLayers[ViewIndex];
-			CompositedProjection.pose = ToXrPose(EyePose, WorldToMeters);
-			CompositedProjection.fov = View.fov;
+			XrCompositionLayerProjectionView& Projection = PipelinedLayerStateRendering.ProjectionLayers[ViewIndex];
+			Projection.pose = ToXrPose(BasePoseTransform, WorldToMeters);
+			Projection.fov = View.fov;
+
+			if (EnumHasAnyFlags(PipelinedLayerStateRendering.LayerStateFlags, EOpenXRLayerStateFlags::SubmitEmulatedFaceLockedLayer))
+			{
+				XrCompositionLayerProjectionView& CompositedProjection = PipelinedLayerStateRendering.EmulatedLayerState.CompositedProjectionLayers[ViewIndex];
+				CompositedProjection.pose = ToXrPose(EyePose, WorldToMeters);
+				CompositedProjection.fov = View.fov;
+			}
 		}
 	}
-
+	
 #if !PLATFORM_HOLOLENS
 	if (bHiddenAreaMaskSupported && bNeedReBuildOcclusionMesh)
 	{
@@ -2891,28 +2894,31 @@ void FOpenXRHMD::OnLateUpdateApplied_RenderThread(FRHICommandListImmediate& RHIC
 
 	ensure(IsInRenderingThread());
 
-	for (int32 ViewIndex = 0; ViewIndex < PipelinedLayerStateRendering.ProjectionLayers.Num(); ViewIndex++)
+	if (PipelinedFrameStateRendering.Views.Num() == PipelinedLayerStateRendering.ProjectionLayers.Num())
 	{
-		const XrView& View = PipelinedFrameStateRendering.Views[ViewIndex];
-		XrCompositionLayerProjectionView& Projection = PipelinedLayerStateRendering.ProjectionLayers[ViewIndex];
+		for (int32 ViewIndex = 0; ViewIndex < PipelinedLayerStateRendering.ProjectionLayers.Num(); ViewIndex++)
+		{
+			const XrView& View = PipelinedFrameStateRendering.Views[ViewIndex];
+			XrCompositionLayerProjectionView& Projection = PipelinedLayerStateRendering.ProjectionLayers[ViewIndex];
 
-		// Apply the new HMD orientation to each eye pose for the final pose
-		FTransform EyePose = ToFTransform(View.pose, GetWorldToMetersScale());
-		FTransform NewRelativePoseTransform = EyePose * NewRelativeTransform;
-		NewRelativePoseTransform.NormalizeRotation();
-		Projection.pose = ToXrPose(NewRelativePoseTransform, GetWorldToMetersScale());
+			// Apply the new HMD orientation to each eye pose for the final pose
+			FTransform EyePose = ToFTransform(View.pose, GetWorldToMetersScale());
+			FTransform NewRelativePoseTransform = EyePose * NewRelativeTransform;
+			NewRelativePoseTransform.NormalizeRotation();
+			Projection.pose = ToXrPose(NewRelativePoseTransform, GetWorldToMetersScale());
 
-		// Update the field-of-view to match the final projection matrix
-		Projection.fov = View.fov;
+			// Update the field-of-view to match the final projection matrix
+			Projection.fov = View.fov;
 
-		if (EnumHasAnyFlags(PipelinedLayerStateRendering.LayerStateFlags, EOpenXRLayerStateFlags::SubmitEmulatedFaceLockedLayer))
-		{
-			XrCompositionLayerProjectionView& CompositedProjection = PipelinedLayerStateRendering.EmulatedLayerState.CompositedProjectionLayers[ViewIndex];
-			CompositedProjection.pose = ToXrPose(EyePose, GetWorldToMetersScale());
-			CompositedProjection.fov = View.fov;
+			if (EnumHasAnyFlags(PipelinedLayerStateRendering.LayerStateFlags, EOpenXRLayerStateFlags::SubmitEmulatedFaceLockedLayer))
+			{
+				XrCompositionLayerProjectionView& CompositedProjection = PipelinedLayerStateRendering.EmulatedLayerState.CompositedProjectionLayers[ViewIndex];
+				CompositedProjection.pose = ToXrPose(EyePose, GetWorldToMetersScale());
+				CompositedProjection.fov = View.fov;
+			}
 		}
 	}
-
+	
 	RHICmdList.EnqueueLambda([this, ProjectionLayers = PipelinedLayerStateRendering.ProjectionLayers, CompositedProjectionLayers = PipelinedLayerStateRendering.EmulatedLayerState.CompositedProjectionLayers](FRHICommandListImmediate& InRHICmdList)
 	{
 		PipelinedLayerStateRHI.ProjectionLayers = ProjectionLayers;
