OpenXR: Move OpenXRPlatformRHI to public

#rnx

diff --git a/Engine/Plugins/Runtime/OpenXR/Source/OpenXRHMD/OpenXRHMD.Build.cs b/Engine/Plugins/Runtime/OpenXR/Source/OpenXRHMD/OpenXRHMD.Build.cs
--- a/Engine/Plugins/Runtime/OpenXR/Source/OpenXRHMD/OpenXRHMD.Build.cs
+++ b/Engine/Plugins/Runtime/OpenXR/Source/OpenXRHMD/OpenXRHMD.Build.cs
@@ -48,27 +48,39 @@ namespace UnrealBuildTool.Rules
 
             if (Target.Platform == UnrealTargetPlatform.Win64)
             {
-                PrivateDependencyModuleNames.AddRange(new string[] {
+				PublicDependencyModuleNames.AddRange(new string[] {
 					"D3D11RHI",
 					"D3D12RHI"
 				});
 
-                AddEngineThirdPartyPrivateStaticDependencies(Target, "DX11", "DX12");
-            }
+				if (!bUsePrecompiled || Target.LinkType == TargetLinkType.Monolithic)
+				{
+					PublicDependencyModuleNames.AddRange(new string[] {
+						"DX11",
+						"DX12"
+					});
+				}
+			}
 
             if (Target.Platform == UnrealTargetPlatform.Win64 || Target.Platform == UnrealTargetPlatform.Android)
             {
-                PrivateDependencyModuleNames.Add("OpenGLDrv");
+				PublicDependencyModuleNames.Add("OpenGLDrv");
 
-                AddEngineThirdPartyPrivateStaticDependencies(Target, "OpenGL");
+				if (!bUsePrecompiled || Target.LinkType == TargetLinkType.Monolithic)
+				{
+					PublicDependencyModuleNames.Add("OpenGL");
+				}
 			}
 
 			if (Target.Platform == UnrealTargetPlatform.Win64 || Target.Platform == UnrealTargetPlatform.Android  
 			    || Target.IsInPlatformGroup(UnrealPlatformGroup.Linux))
             {
-                PrivateDependencyModuleNames.Add("VulkanRHI");
+				PublicDependencyModuleNames.Add("VulkanRHI");
 
-                AddEngineThirdPartyPrivateStaticDependencies(Target, "Vulkan");
+				if (!bUsePrecompiled || Target.LinkType == TargetLinkType.Monolithic)
+				{
+					PublicDependencyModuleNames.Add("Vulkan");
+				}
 			}
 
 			if (Target.Platform == UnrealTargetPlatform.Android)
diff --git a/Engine/Plugins/Runtime/OpenXR/Source/OpenXRHMD/Private/OpenXRPlatformRHI.h b/Engine/Plugins/Runtime/OpenXR/Source/OpenXRHMD/Private/OpenXRPlatformRHI.h
deleted file mode 100644
diff --git a/Engine/Plugins/Runtime/OpenXR/Source/OpenXRHMD/Public/OpenXRPlatformRHI.h b/Engine/Plugins/Runtime/OpenXR/Source/OpenXRHMD/Public/OpenXRPlatformRHI.h
new file mode 100644
--- /dev/null
+++ b/Engine/Plugins/Runtime/OpenXR/Source/OpenXRHMD/Public/OpenXRPlatformRHI.h
@@ -0,0 +1,72 @@
+// Copyright Epic Games, Inc. All Rights Reserved.
+
+#pragma once
+
+#include "HAL/Platform.h"
+
+#if PLATFORM_WINDOWS || PLATFORM_HOLOLENS
+#define XR_USE_PLATFORM_WIN32		1
+#define XR_USE_GRAPHICS_API_D3D11	1
+#define XR_USE_GRAPHICS_API_D3D12	1
+#endif
+
+#if PLATFORM_WINDOWS
+#define XR_USE_GRAPHICS_API_OPENGL	1
+#endif
+
+#if PLATFORM_ANDROID
+#define XR_USE_PLATFORM_ANDROID 1
+#define XR_USE_GRAPHICS_API_OPENGL_ES	1
+#endif
+
+#if PLATFORM_WINDOWS || PLATFORM_ANDROID || PLATFORM_LINUX
+#define XR_USE_GRAPHICS_API_VULKAN 1
+#endif
+
+//-------------------------------------------------------------------------------------------------
+// D3D11
+//-------------------------------------------------------------------------------------------------
+
+#ifdef XR_USE_GRAPHICS_API_D3D11
+#include "ID3D11DynamicRHI.h"
+#endif // XR_USE_GRAPHICS_API_D3D11
+
+
+//-------------------------------------------------------------------------------------------------
+// D3D12
+//-------------------------------------------------------------------------------------------------
+
+#ifdef XR_USE_GRAPHICS_API_D3D12
+#include "ID3D12DynamicRHI.h"
+#endif // XR_USE_GRAPHICS_API_D3D12
+
+
+//-------------------------------------------------------------------------------------------------
+// OpenGL
+//-------------------------------------------------------------------------------------------------
+
+#if defined(XR_USE_GRAPHICS_API_OPENGL) || defined(XR_USE_GRAPHICS_API_OPENGL_ES)
+#include "IOpenGLDynamicRHI.h"
+#endif // XR_USE_GRAPHICS_API_OPENGL
+
+//-------------------------------------------------------------------------------------------------
+// Vulkan
+//-------------------------------------------------------------------------------------------------
+
+#ifdef XR_USE_GRAPHICS_API_VULKAN
+#include "IVulkanDynamicRHI.h"
+#endif // XR_USE_GRAPHICS_API_VULKAN
+
+#if PLATFORM_WINDOWS
+#include "Windows/AllowWindowsPlatformTypes.h"
+#endif
+
+#if PLATFORM_ANDROID
+#include <android_native_app_glue.h>
+#endif
+
+#include <openxr/openxr_platform.h>
+
+#if PLATFORM_WINDOWS
+#include "Windows/HideWindowsPlatformTypes.h"
+#endif
diff --git a/Engine/Plugins/Runtime/XR/XRScribe/Source/XRScribe/XRScribe.Build.cs b/Engine/Plugins/Runtime/XR/XRScribe/Source/XRScribe/XRScribe.Build.cs
--- a/Engine/Plugins/Runtime/XR/XRScribe/Source/XRScribe/XRScribe.Build.cs
+++ b/Engine/Plugins/Runtime/XR/XRScribe/Source/XRScribe/XRScribe.Build.cs
@@ -25,50 +25,6 @@ namespace UnrealBuildTool.Rules
                 "DeveloperSettings",
 				"Projects",
 			});
-
-			// To get access to OpenXRPlatformRHI
-			PublicIncludePaths.Add(System.IO.Path.Combine(GetModuleDirectory("OpenXRHMD"), "Private"));
-
-			// For the RHI/SDK modules, we'll make both public so XRScribeTests can access them.
-			// We could just make the SDKs private, and have the tests include the headers. But we
-			// might actually want to do something in the tests eventually.
-
-			if (Target.Platform == UnrealTargetPlatform.Win64)
-			{
-				PublicDependencyModuleNames.AddRange(new string[] {
-					"D3D11RHI",
-					"D3D12RHI"
-				});
-
-				if (!bUsePrecompiled || Target.LinkType == TargetLinkType.Monolithic)
-				{
-					PublicDependencyModuleNames.AddRange(new string[] {
-						"DX11",
-						"DX12"
-					});
-				}
-			}
-
-			if (Target.Platform == UnrealTargetPlatform.Win64 || Target.Platform == UnrealTargetPlatform.Android)
-			{
-				PublicDependencyModuleNames.Add("OpenGLDrv");
-
-				if (!bUsePrecompiled || Target.LinkType == TargetLinkType.Monolithic)
-				{
-					PublicDependencyModuleNames.Add("OpenGL");
-				}
-			}
-
-			if (Target.Platform == UnrealTargetPlatform.Win64 || Target.Platform == UnrealTargetPlatform.Android
-				|| Target.IsInPlatformGroup(UnrealPlatformGroup.Linux))
-			{
-				PublicDependencyModuleNames.Add("VulkanRHI");
-
-				if (!bUsePrecompiled || Target.LinkType == TargetLinkType.Monolithic)
-				{
-					PublicDependencyModuleNames.Add("Vulkan");
-				}
-			}
 		}
 	}
 }
