OpenXR related private header access cleanup.
-Refactoring to move some cross-module included Private header files into Internal.  Eliminating some includes by refactoring.  Eliminating some simply unnecessary includes.

diff --git a/Engine/Plugins/Runtime/OpenXR/Source/OpenXRHMD/Internal/Epic_openxr.h b/Engine/Plugins/Runtime/OpenXR/Source/OpenXRHMD/Internal/Epic_openxr.h
new file mode 100644
--- /dev/null
+++ b/Engine/Plugins/Runtime/OpenXR/Source/OpenXRHMD/Internal/Epic_openxr.h
@@ -0,0 +1,38 @@
+// Copyright Epic Games, Inc. All Rights Reserved.
+
+#pragma once
+
+#include <openxr/openxr.h>
+
+// Uncomment to make intellisense work better.
+//#include "/../../../../../Source/ThirdParty/OpenXR/include/openxr/openxr.h"
+
+// This number is in a block that is unused by openxr, and that seemed unlikely to be used soon based on the patterns of usage.
+// But if it ever was used by openxr we would have to change it.
+#define XR_TYPE_EPIC 1000059050
+
+typedef enum XrStructureTypeEPIC {
+	XR_TYPE_SPACE_ACCELERATION_EPIC = XR_TYPE_EPIC,
+	XR_STRUCTURE_TYPE_EPIC_MAX_ENUM = 0x7FFFFFFF
+} XrStructureTypeEPIC;
+
+
+//  Extensions
+
+#define XR_EPIC_space_acceleration 1
+#define XR_EPIC_space_acceleration_SPEC_VERSION 1
+#define XR_EPIC_SPACE_ACCELERATION_NAME "XR_EPIC_space_acceleration"
+
+typedef XrFlags64 XrSpaceAccelerationFlagsEPIC;
+
+// Flag bits for XrSpaceAccelerationFlagsEPIC
+static const XrSpaceAccelerationFlagsEPIC XR_SPACE_ACCELERATION_LINEAR_VALID_BIT_EPIC = 0x00000001;
+static const XrSpaceAccelerationFlagsEPIC XR_SPACE_ACCELERATION_ANGULAR_VALID_BIT_EPIC = 0x00000002;
+
+typedef struct XrSpaceAccelerationEPIC {
+	XrStructureType							type;
+	void* XR_MAY_ALIAS						next;
+	XrSpaceAccelerationFlagsEPIC			accelerationFlags;
+	XrVector3f								linearAcceleration;  // m/s^2
+	XrVector3f							angularAcceleration; // rad/s^2
+} XrSpaceAccelerationEPIC;
diff --git a/Engine/Plugins/Runtime/OpenXR/Source/OpenXRHMD/Internal/OpenXRHMD_RenderBridge.h b/Engine/Plugins/Runtime/OpenXR/Source/OpenXRHMD/Internal/OpenXRHMD_RenderBridge.h
new file mode 100644
--- /dev/null
+++ b/Engine/Plugins/Runtime/OpenXR/Source/OpenXRHMD/Internal/OpenXRHMD_RenderBridge.h
@@ -0,0 +1,82 @@
+// Copyright Epic Games, Inc. All Rights Reserved.
+
+#pragma once
+
+#include "CoreMinimal.h"
+#include "HeadMountedDisplayBase.h"
+#include "XRRenderBridge.h"
+#include "XRSwapChain.h"
+#include "OpenXRPlatformRHI.h"
+
+#include <openxr/openxr.h>
+
+class FOpenXRHMD;
+class FRHICommandListImmediate;
+
+class FOpenXRRenderBridge : public FXRRenderBridge
+{
+public:
+	FOpenXRRenderBridge(XrInstance InInstance)
+		: Instance(InInstance)
+		, OpenXRHMD(nullptr)
+	{ }
+
+	void SetOpenXRHMD(FOpenXRHMD* InHMD) { OpenXRHMD = InHMD; }
+	virtual uint64 GetGraphicsAdapterLuid(XrSystemId InSystem) { return 0; };
+	virtual void* GetGraphicsBinding(XrSystemId InSystem) = 0;
+
+	virtual FXRSwapChainPtr CreateSwapchain(XrSession InSession, uint8 Format, uint8& OutActualFormat, uint32 SizeX, uint32 SizeY, uint32 ArraySize, uint32 NumMips, uint32 NumSamples, ETextureCreateFlags CreateFlags, const FClearValueBinding& ClearValueBinding, ETextureCreateFlags AuxiliaryCreateFlags = ETextureCreateFlags::None) = 0;
+
+	FXRSwapChainPtr CreateSwapchain(XrSession InSession, FRHITexture2D* Template, ETextureCreateFlags CreateFlags)
+	{
+		if (!Template)
+		{
+			return nullptr;
+		}
+
+		uint8 UnusedOutFormat = 0;
+		return CreateSwapchain(InSession,
+			Template->GetFormat(),
+			UnusedOutFormat,
+			Template->GetSizeX(),
+			Template->GetSizeY(),
+			1,
+			Template->GetNumMips(),
+			Template->GetNumSamples(),
+			Template->GetFlags() | CreateFlags,
+			Template->GetClearBinding());
+	}
+
+	/** FRHICustomPresent */
+	virtual bool Present(int32& InOutSyncInterval) override;
+
+	virtual bool Support10BitSwapchain() const { return false; }
+
+	virtual bool HDRGetMetaDataForStereo(EDisplayOutputFormat& OutDisplayOutputFormat, EDisplayColorGamut& OutDisplayColorGamut, bool& OutbHDRSupported) { return false; }
+
+	virtual void SetSkipRate(uint32 SkipRate) {}
+
+	virtual void HMDOnFinishRendering_RHIThread();
+
+protected:
+	XrInstance Instance;
+	FOpenXRHMD* OpenXRHMD;
+
+private:
+};
+
+#ifdef XR_USE_GRAPHICS_API_D3D11
+FOpenXRRenderBridge* CreateRenderBridge_D3D11(XrInstance InInstance);
+#endif
+#ifdef XR_USE_GRAPHICS_API_D3D12
+FOpenXRRenderBridge* CreateRenderBridge_D3D12(XrInstance InInstance);
+#endif
+#ifdef XR_USE_GRAPHICS_API_OPENGL_ES
+FOpenXRRenderBridge* CreateRenderBridge_OpenGLES(XrInstance InInstance);
+#endif
+#ifdef XR_USE_GRAPHICS_API_OPENGL
+FOpenXRRenderBridge* CreateRenderBridge_OpenGL(XrInstance InInstance);
+#endif
+#ifdef XR_USE_GRAPHICS_API_VULKAN
+FOpenXRRenderBridge* CreateRenderBridge_Vulkan(XrInstance InInstance);
+#endif
diff --git a/Engine/Plugins/Runtime/OpenXR/Source/OpenXRHMD/Internal/OpenXRHMD_Swapchain.h b/Engine/Plugins/Runtime/OpenXR/Source/OpenXRHMD/Internal/OpenXRHMD_Swapchain.h
new file mode 100644
--- /dev/null
+++ b/Engine/Plugins/Runtime/OpenXR/Source/OpenXRHMD/Internal/OpenXRHMD_Swapchain.h
@@ -0,0 +1,51 @@
+// Copyright Epic Games, Inc. All Rights Reserved.
+
+#pragma once
+
+#include "CoreMinimal.h"
+#include "HeadMountedDisplayBase.h"
+#include "XRSwapChain.h"
+#include "OpenXRPlatformRHI.h"
+
+#include <openxr/openxr.h>
+
+class FOpenXRSwapchain : public FXRSwapChain
+{
+public:
+	FOpenXRSwapchain(TArray<FTextureRHIRef>&& InRHITextureSwapChain, const FTextureRHIRef & InRHITexture, XrSwapchain InHandle);
+	virtual ~FOpenXRSwapchain();
+
+	virtual void IncrementSwapChainIndex_RHIThread() override final;
+	virtual void WaitCurrentImage_RHIThread(int64 Timeout) override final;
+	virtual void ReleaseCurrentImage_RHIThread() override final;
+
+	XrSwapchain GetHandle() { return Handle; }
+	static XrSwapchain CreateSwapchain(XrSession InSession, uint32 PlatformFormat, uint32 SizeX, uint32 SizeY, uint32 ArraySize, uint32 NumMips, uint32 NumSamples, ETextureCreateFlags CreateFlags, void* Next = nullptr);
+	static uint8 GetNearestSupportedSwapchainFormat(XrSession InSession, uint8 RequestedFormat, TFunction<uint32(uint8)> ToPlatformFormat = nullptr);
+	void GetFragmentDensityMaps(TArray<FTextureRHIRef>& OutTextureChain, const bool bIsMobileMultiViewEnabled);
+
+protected:
+	XrSwapchain Handle;
+
+	/** Whether the image associated with the swapchain has been acquired. */
+	std::atomic<bool> ImageAcquired;
+
+	/** Whether the image associated with the swapchain is ready for being written to. */
+	std::atomic<bool> ImageReady;
+};
+
+#ifdef XR_USE_GRAPHICS_API_D3D11
+FXRSwapChainPtr CreateSwapchain_D3D11(XrSession InSession, uint8 Format, uint8& OutActualFormat, uint32 SizeX, uint32 SizeY, uint32 ArraySize, uint32 NumMips, uint32 NumSamples, ETextureCreateFlags CreateFlags, const FClearValueBinding& ClearValueBinding, ETextureCreateFlags AuxiliaryCreateFlags = ETextureCreateFlags::None);
+#endif
+#ifdef XR_USE_GRAPHICS_API_D3D12
+FXRSwapChainPtr CreateSwapchain_D3D12(XrSession InSession, uint8 Format, uint8& OutActualFormat, uint32 SizeX, uint32 SizeY, uint32 ArraySize, uint32 NumMips, uint32 NumSamples, ETextureCreateFlags CreateFlags, const FClearValueBinding& ClearValueBinding, ETextureCreateFlags AuxiliaryCreateFlags = ETextureCreateFlags::None);
+#endif
+#ifdef XR_USE_GRAPHICS_API_OPENGL
+FXRSwapChainPtr CreateSwapchain_OpenGL(XrSession InSession, uint8 Format, uint8& OutActualFormat, uint32 SizeX, uint32 SizeY, uint32 ArraySize, uint32 NumMips, uint32 NumSamples, ETextureCreateFlags CreateFlags, const FClearValueBinding& ClearValueBinding, ETextureCreateFlags AuxiliaryCreateFlags = ETextureCreateFlags::None);
+#endif
+#ifdef XR_USE_GRAPHICS_API_OPENGL_ES
+FXRSwapChainPtr CreateSwapchain_OpenGLES(XrSession InSession, uint8 Format, uint8& OutActualFormat, uint32 SizeX, uint32 SizeY, uint32 ArraySize, uint32 NumMips, uint32 NumSamples, ETextureCreateFlags CreateFlags, const FClearValueBinding& ClearValueBinding, ETextureCreateFlags AuxiliaryCreateFlags = ETextureCreateFlags::None);
+#endif
+#ifdef XR_USE_GRAPHICS_API_VULKAN
+FXRSwapChainPtr CreateSwapchain_Vulkan(XrSession InSession, uint8 Format, uint8& OutActualFormat, uint32 SizeX, uint32 SizeY, uint32 ArraySize, uint32 NumMips, uint32 NumSamples, ETextureCreateFlags CreateFlags, const FClearValueBinding& ClearValueBinding, ETextureCreateFlags AuxiliaryCreateFlags = ETextureCreateFlags::None);
+#endif
\ No newline at end of file
diff --git a/Engine/Plugins/Runtime/OpenXR/Source/OpenXRHMD/Private/Epic_openxr.h b/Engine/Plugins/Runtime/OpenXR/Source/OpenXRHMD/Private/Epic_openxr.h
deleted file mode 100644
diff --git a/Engine/Plugins/Runtime/OpenXR/Source/OpenXRHMD/Private/OpenXRHMD_RenderBridge.cpp b/Engine/Plugins/Runtime/OpenXR/Source/OpenXRHMD/Private/OpenXRHMD_RenderBridge.cpp
--- a/Engine/Plugins/Runtime/OpenXR/Source/OpenXRHMD/Private/OpenXRHMD_RenderBridge.cpp
+++ b/Engine/Plugins/Runtime/OpenXR/Source/OpenXRHMD/Private/OpenXRHMD_RenderBridge.cpp
@@ -11,7 +11,7 @@ bool FOpenXRRenderBridge::Present(int32& InOutSyncInterval)
 
 	if (OpenXRHMD)
 	{
-		OpenXRHMD->OnFinishRendering_RHIThread();
+		HMDOnFinishRendering_RHIThread();
 		bNeedsNativePresent = !OpenXRHMD->IsStandaloneStereoOnlyDevice();
 	}
 
@@ -20,6 +20,14 @@ bool FOpenXRRenderBridge::Present(int32& InOutSyncInterval)
 	return bNeedsNativePresent;
 }
 
+void FOpenXRRenderBridge::HMDOnFinishRendering_RHIThread()
+{
+	if (OpenXRHMD)
+	{
+		OpenXRHMD->OnFinishRendering_RHIThread();
+	}
+}
+
 #ifdef XR_USE_GRAPHICS_API_D3D11
 class FD3D11RenderBridge : public FOpenXRRenderBridge
 {
diff --git a/Engine/Plugins/Runtime/OpenXR/Source/OpenXRHMD/Private/OpenXRHMD_RenderBridge.h b/Engine/Plugins/Runtime/OpenXR/Source/OpenXRHMD/Private/OpenXRHMD_RenderBridge.h
deleted file mode 100644
diff --git a/Engine/Plugins/Runtime/OpenXR/Source/OpenXRHMD/Private/OpenXRHMD_Swapchain.h b/Engine/Plugins/Runtime/OpenXR/Source/OpenXRHMD/Private/OpenXRHMD_Swapchain.h
deleted file mode 100644
diff --git a/Engine/Plugins/Runtime/OpenXR/Source/OpenXRHMD/Public/OpenXRCore.h b/Engine/Plugins/Runtime/OpenXR/Source/OpenXRHMD/Public/OpenXRCore.h
--- a/Engine/Plugins/Runtime/OpenXR/Source/OpenXRHMD/Public/OpenXRCore.h
+++ b/Engine/Plugins/Runtime/OpenXR/Source/OpenXRHMD/Public/OpenXRCore.h
@@ -293,3 +293,42 @@ FORCEINLINE void FilterActionName(const char* InActionName, char* OutActionName)
 	}
 	OutActionName[i] = '\0';
 }
+
+namespace OpenXR
+{
+	template<typename T>
+	T* FindChainedStructByType(void* Head, XrStructureType XRType)
+	{
+		XrBaseOutStructure* Ptr = (XrBaseOutStructure*)Head;
+		while (Ptr)
+		{
+			if (Ptr->type == XRType)
+			{
+				return (T*)Ptr;
+			}
+			else
+			{
+				Ptr = (XrBaseOutStructure*)Ptr->next;
+			}
+		}
+		return nullptr;
+	}
+
+	template<typename T>
+	const T* FindChainedStructByType(const void* Head, XrStructureType XRType)
+	{
+		const XrBaseInStructure* Ptr = (XrBaseInStructure*)Head;
+		while (Ptr)
+		{
+			if (Ptr->type == XRType)
+			{
+				return (const T*)Ptr;
+			}
+			else
+			{
+				Ptr = (XrBaseInStructure*)Ptr->next;
+			}
+		}
+		return nullptr;
+	}
+}
\ No newline at end of file
