Fixes face-locked stereo layers (such as the console and stat displays) becoming offset from the user's view after teleporting or otherwise changing the HMD's TrackingToWorld transform.

diff --git a/Engine/Plugins/Runtime/OpenXR/Source/OpenXRHMD/Private/OpenXRHMD_Layer.cpp b/Engine/Plugins/Runtime/OpenXR/Source/OpenXRHMD/Private/OpenXRHMD_Layer.cpp
--- a/Engine/Plugins/Runtime/OpenXR/Source/OpenXRHMD/Private/OpenXRHMD_Layer.cpp
+++ b/Engine/Plugins/Runtime/OpenXR/Source/OpenXRHMD/Private/OpenXRHMD_Layer.cpp
@@ -90,15 +90,15 @@ TArray<FXrCompositionLayerUnion> FOpenXRLayer::CreateOpenXRLayer(FTransform InvT
 
 	if (Desc.HasShape<FQuadLayer>())
 	{
-		CreateOpenXRQuadLayer(bIsStereo, bNoAlpha, InvTrackingToWorld, WorldToMeters, Space, Headers);
+		CreateOpenXRQuadLayer(bIsStereo, bNoAlpha, PositionTransform, WorldToMeters, Space, Headers);
 	}
 	else if (Desc.HasShape<FCylinderLayer>())
 	{
-		CreateOpenXRCylinderLayer(bIsStereo, bNoAlpha, InvTrackingToWorld, WorldToMeters, Space, Headers);
+		CreateOpenXRCylinderLayer(bIsStereo, bNoAlpha, PositionTransform, WorldToMeters, Space, Headers);
 	}
 	else if (Desc.HasShape<FEquirectLayer>())
 	{
-		CreateOpenXREquirectLayer(bIsStereo, bNoAlpha, InvTrackingToWorld, WorldToMeters, Space, Headers);
+		CreateOpenXREquirectLayer(bIsStereo, bNoAlpha, PositionTransform, WorldToMeters, Space, Headers);
 	}
 
 	return Headers;
