UE-194257 OpenXR: vive hand tracking layer breaks oculus runtime
-Added a ProblematicOpenXRApiLayerInfos ini setting that defines specific ApiLayer versions that are behaving badly with regard to an extension.
-Attempt to create the openxr instance by repeatedly calling xrCreateInstance with progressively more of the potentially problematic extensions not enabled.  We will either find a combination that works or never find a combination that works.
-This Cl also provides a structure for implementing further problematic openxr api layer handling, should a project need further workarounds similar to this.

diff --git a/Engine/Plugins/Runtime/OpenXR/Config/Engine.ini b/Engine/Plugins/Runtime/OpenXR/Config/Engine.ini
new file mode 100644
--- /dev/null
+++ b/Engine/Plugins/Runtime/OpenXR/Config/Engine.ini
@@ -0,0 +1,2 @@
+[OpenXR]
++ProblematicOpenXRApiLayerInfos="Name=XR_APILAYER_VIVE_hand_tracking,MinVersion=2,MaxVersion=2,ExtensionAddedToFallbackWithout=XR_EXT_hand_tracking"
diff --git a/Engine/Plugins/Runtime/OpenXR/Source/OpenXRHMD/Private/OpenXRCore.cpp b/Engine/Plugins/Runtime/OpenXR/Source/OpenXRHMD/Private/OpenXRCore.cpp
--- a/Engine/Plugins/Runtime/OpenXR/Source/OpenXRHMD/Private/OpenXRCore.cpp
+++ b/Engine/Plugins/Runtime/OpenXR/Source/OpenXRHMD/Private/OpenXRCore.cpp
@@ -96,4 +96,16 @@ bool InitOpenXRCore(XrInstance Instance)
 	return true;
 }
 
+void EnumerateOpenXRApiLayers(TArray<XrApiLayerProperties>& OutProperties)
+{
+	uint32_t Count = 0;
+	XR_ENSURE(xrEnumerateApiLayerProperties(0, &Count, nullptr));
+	OutProperties.SetNum(Count);
+	for (auto& Prop : OutProperties)
+	{
+		Prop = XrApiLayerProperties{ XR_TYPE_API_LAYER_PROPERTIES };
+	}
+	XR_ENSURE(xrEnumerateApiLayerProperties(Count, &Count, OutProperties.GetData()));
+}
+
 #undef GET_XR_ENTRYPOINTS
diff --git a/Engine/Plugins/Runtime/OpenXR/Source/OpenXRHMD/Private/OpenXRHMDModule.cpp b/Engine/Plugins/Runtime/OpenXR/Source/OpenXRHMD/Private/OpenXRHMDModule.cpp
--- a/Engine/Plugins/Runtime/OpenXR/Source/OpenXRHMD/Private/OpenXRHMDModule.cpp
+++ b/Engine/Plugins/Runtime/OpenXR/Source/OpenXRHMD/Private/OpenXRHMDModule.cpp
@@ -7,6 +7,7 @@
 #include "OpenXRCore.h"
 #include "IOpenXRExtensionPlugin.h"
 #include "IOpenXRARModule.h"
+#include "OpenXRHMDSettings.h"
 #include "BuildSettings.h"
 #include "GeneralProjectSettings.h"
 #include "Epic_openxr.h"
@@ -197,6 +198,7 @@ bool FOpenXRHMDModule::EnumerateExtensions()
 	if (XR_FAILED(xrEnumerateInstanceExtensionProperties(nullptr, 0, &ExtensionsCount, nullptr)))
 	{
 		// If it fails this early that means there's no runtime installed
+		UE_LOG(LogHMD, Log, TEXT("xrEnumerateInstanceExtensionProperties failed, suggests no runtime is installed."));
 		return false;
 	}
 
@@ -736,19 +738,8 @@ bool FOpenXRHMDModule::InitInstance()
 	Info.next = &InstanceCreateInfoAndroid;
 #endif // PLATFORM_ANDROID
 
-	for (IOpenXRExtensionPlugin* Module : ExtensionPlugins)
-	{
-		Info.next = Module->OnCreateInstance(this, Info.next);
-	}
-
-	XrResult Result = xrCreateInstance(&Info, &Instance);
-	if (XR_FAILED(Result))
+	if (!TryCreateInstance(Info))
 	{
-		UE_LOG(LogHMD, Log, TEXT("Failed to create an OpenXR instance, result is %s. Please check if you have an OpenXR runtime installed. The following extensions were enabled:"), OpenXRResultToString(Result));
-		for (const char* Extension : EnabledExtensions)
-		{
-			UE_LOG(LogHMD, Log, TEXT("- %S"), Extension);
-		}
 		return false;
 	}
 
@@ -787,6 +778,271 @@ bool FOpenXRHMDModule::InitInstance()
 	return true;
 }
 
+struct FProblematicOpenXRApiLayerInfo
+{
+	FProblematicOpenXRApiLayerInfo(const FString& String)
+	{
+		FString StringCopy = String;
+		StringCopy.RemoveSpacesInline();
+		FParse::Value(*String, TEXT("Name="), Name);
+		FParse::Value(*String, TEXT("MinVersion="), MinVersion);
+		FParse::Value(*String, TEXT("MaxVersion="), MaxVersion);
+		FParse::Value(*String, TEXT("ExtensionAddedToFallbackWithout="), ExtensionAddedToFallbackWithout);
+	}
+	FString Name;
+	int MinVersion = 0;
+	int MaxVersion = 0; // 0 for both = all.
+	FString ExtensionAddedToFallbackWithout;  // If this extension is available and instance creation fails with XR_ERROR_EXTENSION_NOT_PRESENT we will try again without it.
+};
+
+bool FOpenXRHMDModule::TryCreateInstance(XrInstanceCreateInfo& Info)
+{
+	// Cache a copy of the info not modified by ExtensionPlugins, some of which we might have to disable.
+	XrInstanceCreateInfo LocalInfo = Info;
+
+	for (IOpenXRExtensionPlugin* Module : ExtensionPlugins)
+	{
+		Info.next = Module->OnCreateInstance(this, Info.next);
+	}
+
+	XrResult Result = xrCreateInstance(&Info, &Instance);
+	if (XR_FAILED(Result))
+	{
+		if (Result == XR_ERROR_EXTENSION_NOT_PRESENT)
+		{
+			// An extension we requested is not supported, but we normally only add extensions that must be supported and extensions that were listed to us by the runtime, so how did we get into this situation?
+			// A badly behaving layer might add an extension into the available extensions list, but then not remove it before the runtime sees that list causing the runtime to fail instance creation with XR_ERROR_EXTENSION_NOT_PRESENT.
+			// Therefore we have a ProblematicOpenXRApiLayerInfos config setting to define layers that might give us bad extensions which we can then try to create an instance without.
+
+			// Read the ini.  This code can run so early that we cannot use a default object to load this more easily.
+			TArray<FProblematicOpenXRApiLayerInfo> ProblematicOpenXRApiLayerInfos;
+			{
+				TArray<FString> ProblematicOpenXRApiLayerInfoStrings;
+				GConfig->GetArray(TEXT("OpenXR"), TEXT("ProblematicOpenXRApiLayerInfos"), ProblematicOpenXRApiLayerInfoStrings, GEngineIni);
+				ProblematicOpenXRApiLayerInfos.Reserve(ProblematicOpenXRApiLayerInfoStrings.Num());
+				for (FString& Str : ProblematicOpenXRApiLayerInfoStrings)
+				{
+					ProblematicOpenXRApiLayerInfos.Emplace(Str);
+				}
+			}
+
+			if (ProblematicOpenXRApiLayerInfos.Num() > 0)
+			{
+				// Copy the currently active ProblematicLayerInfos that have ExtensionsAddedToFallbackWithout into a new array.
+				TArray<const FProblematicOpenXRApiLayerInfo*> ActiveProblematicLayerInfos;
+				{
+					TArray<XrApiLayerProperties> ActiveLayerProperties;
+					EnumerateOpenXRApiLayers(ActiveLayerProperties);
+					TMap<FString, int> ActiveLayerPropertyNameMap;
+					for (int i = 0; i < ActiveLayerProperties.Num(); ++i)
+					{
+						ActiveLayerPropertyNameMap.Add(ActiveLayerProperties[i].layerName, ActiveLayerProperties[i].layerVersion);
+					}
+					for (const FProblematicOpenXRApiLayerInfo& ProblematicLayerInfo : ProblematicOpenXRApiLayerInfos)
+					{
+						// Skip if no extension specified.
+						if (ProblematicLayerInfo.ExtensionAddedToFallbackWithout.Len() == 0)
+						{
+							continue;
+						}
+						// Check if this one is active.
+						int* Found = ActiveLayerPropertyNameMap.Find(ProblematicLayerInfo.Name);
+						if (Found)
+						{
+							if ((ProblematicLayerInfo.MinVersion == 0 && ProblematicLayerInfo.MaxVersion == 0) || (*Found >= ProblematicLayerInfo.MinVersion && *Found <= ProblematicLayerInfo.MaxVersion))
+							{
+								ActiveProblematicLayerInfos.Add(&ProblematicLayerInfo);
+							}
+						}
+					}
+				}
+
+				if (ActiveProblematicLayerInfos.Num() > 0)
+				{
+					UE_LOG(LogHMD, Log, TEXT("Failed to create an OpenXR instance with all of the extensions that are enabled.  The OpenXR runtime says it does not support all of those. "
+						"We we have a list of Problematic OpenXR ApiLayers that may be adding extensions that are not supported by the runtime. "  
+						"We can try disabling those extensions to see if we can find a set that the runtime will create an instance for. "
+						"Sadly this is a combinatorial process."));
+					if (UE_LOG_ACTIVE(LogHMD, Log))
+					{
+						UE_LOG(LogHMD, Log, TEXT("Problematic Layer Extension Information for currently active OpenXRApiLayers:"));
+						for (const FProblematicOpenXRApiLayerInfo* LayerInfo : ActiveProblematicLayerInfos)
+						{
+							UE_LOG(LogHMD, Log, TEXT("  Layer: %s  Version: %i-%i  Extension: %s"), *LayerInfo->Name, LayerInfo->MinVersion, LayerInfo->MaxVersion, *LayerInfo->ExtensionAddedToFallbackWithout);
+						}
+					}
+
+					// Build a list of all of the problematic extensions, as ansi char pointers from the EnabledExtensions list.
+					TArray<const char*> ProblematicExtensionList;
+					{
+						TSet<const FString*> ProblematicExtensionSet;
+						for (const FProblematicOpenXRApiLayerInfo* LayerInfo : ActiveProblematicLayerInfos)
+						{
+							ProblematicExtensionSet.Add(&(LayerInfo->ExtensionAddedToFallbackWithout));
+						}
+						ProblematicExtensionList.Reserve(ProblematicExtensionSet.Num());
+						for (const FString* ProblematicExtension : ProblematicExtensionSet)
+						{
+							auto ProblematicExtensionConverter = StringCast<ANSICHAR>(**ProblematicExtension);
+							const char* ProblematicExtensionCStr = ProblematicExtensionConverter.Get();
+							for (uint32_t i = 0; i < Info.enabledExtensionCount; ++i)
+							{
+								const char* EnabledExtension = Info.enabledExtensionNames[i];
+								if (FCStringAnsi::Strncmp(ProblematicExtensionCStr, EnabledExtension, XR_MAX_EXTENSION_NAME_SIZE) == 0)
+								{
+									ProblematicExtensionList.Add(EnabledExtension);
+								}
+							}
+						}
+					}
+
+					// Create a list of enabled extensions excluding the problematic ones.
+					TArray<const char*> LocalEnabledExtensions;
+					LocalEnabledExtensions.Reserve(Info.enabledExtensionCount);
+					for (uint32_t i = 0; i < Info.enabledExtensionCount; i++)
+					{
+						const char* EnabledExtension = Info.enabledExtensionNames[i];
+						if (!ProblematicExtensionList.Contains(EnabledExtension))
+						{
+							LocalEnabledExtensions.Add(EnabledExtension);
+						}
+					}
+
+					// Try to xrCreateInstance with each combination until one succeeds or all fail							
+				
+					// So now we have a list of non-problematic extensions (LocalEnabledExtensions) and a list of problematic ones (ProblematicExtensionList).  
+					// We want to try enabling the problematic extensions, exploring all combinations, and last we will try enabling none of them.
+					const int StartIndex = LocalEnabledExtensions.Num();
+					const int NumProblematicExtensions = ProblematicExtensionList.Num();
+					check(NumProblematicExtensions < 64);
+					uint64 Bitfield = (1LL << (NumProblematicExtensions)) - 1;  // Get a 1 for each extension.
+					do
+					{
+						--Bitfield; // We do this first because we already tried with all enabled.
+
+						// Enable some of the problematic extensions
+						LocalEnabledExtensions.SetNum(StartIndex, false); // Shrink off problematic ones from previous iterations
+						for (int i = 0; i < NumProblematicExtensions; ++i)
+						{
+							if (((Bitfield >> i) & 1) == 1)
+							{
+								LocalEnabledExtensions.Add(ProblematicExtensionList[i]);
+							}
+						}
+
+						// Fill in the Instance creation info struct
+						LocalInfo.enabledExtensionCount = LocalEnabledExtensions.Num();
+						LocalInfo.enabledExtensionNames = LocalEnabledExtensions.GetData();
+
+						// Figure out which extensionplugins need to be disabled because their required extensions are disabled.
+						TArray<IOpenXRExtensionPlugin*> LocalExtensionPlugins = ExtensionPlugins;
+						TArray<IOpenXRExtensionPlugin*> LocallyDisabledExtensionPlugins;
+						for (int i = LocalExtensionPlugins.Num() - 1; i >= 0 ; --i)
+						{
+							IOpenXRExtensionPlugin* ExtensionPlugin = LocalExtensionPlugins[i];
+
+							TArray<const ANSICHAR*> RequiredExtensions;
+							if (!ExtensionPlugin->GetRequiredExtensions(RequiredExtensions))
+							{
+								UE_LOG(LogHMD, Error, TEXT("Could not get required OpenXR extensions."));
+								return false;
+							}
+							bool bAllRequiredAreEnabled = true;
+							for (const ANSICHAR* Required : RequiredExtensions)
+							{
+								if (!LocalEnabledExtensions.Contains(Required))
+								{
+									bAllRequiredAreEnabled = false;
+									LocallyDisabledExtensionPlugins.Add(ExtensionPlugin);
+									break;
+								}
+							}
+							if (!bAllRequiredAreEnabled)
+							{
+								LocalExtensionPlugins.RemoveAt(i);
+							}
+						}
+
+						// Log what we are trying now.
+						if (UE_LOG_ACTIVE(LogHMD, Log))
+						{
+							UE_LOG(LogHMD, Log, TEXT("Attempting to create OpenXR instance with some extensions disabled. The following extensions were enabled:"), OpenXRResultToString(Result));
+							for (const char* Extension : EnabledExtensions)
+							{
+								UE_LOG(LogHMD, Log, TEXT("- %S"), Extension);
+							}
+							UE_LOG(LogHMD, Log, TEXT("The following Extensions were disabled because of the ProblematicOpenXRApiLayerInfos:"));
+							for (int i = 0; i < NumProblematicExtensions; ++i)
+							{
+								if (((Bitfield >> i) & 1) == 0)
+								{
+									UE_LOG(LogHMD, Log, TEXT("- %hs"), ProblematicExtensionList[i]);
+								}
+							}
+							UE_LOG(LogHMD, Log, TEXT("The following OpenXRExtensionPlugins were disabled because we disabled some or all of their required extensions:"));
+							for (IOpenXRExtensionPlugin* DisabledPlugin : LocallyDisabledExtensionPlugins)
+							{
+								UE_LOG(LogHMD, Log, TEXT("- %s"), *DisabledPlugin->GetDisplayName());
+							}
+						}
+
+						// Create!
+						for (IOpenXRExtensionPlugin* Module : LocalExtensionPlugins)
+						{
+							LocalInfo.next = Module->OnCreateInstance(this, LocalInfo.next);
+						}
+						XrResult Result2 = xrCreateInstance(&LocalInfo, &Instance);
+						if (XR_SUCCEEDED(Result2))
+						{
+							if (UE_LOG_ACTIVE(LogHMD, Log))
+							{
+								UE_LOG(LogHMD, Log, TEXT("Successfully created OpenXR Instance with the following Extensions disabled because of the ProblematicOpenXRApiLayerInfos:"));
+								for (int i = 0; i < NumProblematicExtensions; ++i)
+								{
+									if (((Bitfield >> i) & 1) == 0)
+									{
+										UE_LOG(LogHMD, Log, TEXT("- %s"), StringCast<TCHAR>(ProblematicExtensionList[i]).Get());
+									}
+								}
+							}
+
+							// Update External data, because we changed it.
+							Info = LocalInfo;
+							EnabledExtensions = LocalEnabledExtensions;
+							ExtensionPlugins = LocalExtensionPlugins;
+							return true;
+						}
+						else
+						{
+							UE_LOG(LogHMD, Log, TEXT("Failed to create an OpenXR instance with some extensions disabled. Result is %s."), OpenXRResultToString(Result));
+						}
+					} while (Bitfield > 0);					
+					
+					UE_LOG(LogHMD, Log, TEXT("We did not find a combination of extensions that works using the ProblematicOpenXRApiLayerInfos.  Instance creation always failed."));
+				}
+			}
+		}
+
+		UE_LOG(LogHMD, Log, TEXT("Failed to create an OpenXR instance, result is %s. Please check if you have an OpenXR runtime installed."), OpenXRResultToString(Result));
+		UE_LOG(LogHMD, Log, TEXT("The following extensions were enabled:"), OpenXRResultToString(Result));
+		for (const char* Extension : EnabledExtensions)
+		{
+			UE_LOG(LogHMD, Log, TEXT("- %S"), Extension);
+		}
+		UE_LOG(LogHMD, Log, TEXT("The following layers were enumerated:"), OpenXRResultToString(Result));
+		TArray<XrApiLayerProperties> EnumeratedLayers;
+		EnumerateOpenXRApiLayers(EnumeratedLayers);
+		for (XrApiLayerProperties& Layer : EnumeratedLayers)
+		{
+			UE_LOG(LogHMD, Log, TEXT("- %S"), Layer.layerName);
+		}
+
+		return false;
+	}
+
+	return true;
+}
+
 XrSystemId FOpenXRHMDModule::GetSystemId() const
 {
 	XrSystemId System = XR_NULL_SYSTEM_ID;
diff --git a/Engine/Plugins/Runtime/OpenXR/Source/OpenXRHMD/Private/OpenXRHMDModule.h b/Engine/Plugins/Runtime/OpenXR/Source/OpenXRHMD/Private/OpenXRHMDModule.h
--- a/Engine/Plugins/Runtime/OpenXR/Source/OpenXRHMD/Private/OpenXRHMDModule.h
+++ b/Engine/Plugins/Runtime/OpenXR/Source/OpenXRHMD/Private/OpenXRHMDModule.h
@@ -69,6 +69,7 @@ private:
 	bool EnumerateLayers();
 	bool InitRenderBridge();
 	bool InitInstance();
+	bool TryCreateInstance(XrInstanceCreateInfo& Info);
 	PFN_xrGetInstanceProcAddr GetDefaultLoader();
 	bool EnableExtensions(const TArray<const ANSICHAR*>& RequiredExtensions, const TArray<const ANSICHAR*>& OptionalExtensions, TArray<const ANSICHAR*>& OutExtensions);
 	bool GetRequiredExtensions(TArray<const ANSICHAR*>& OutExtensions);
diff --git a/Engine/Plugins/Runtime/OpenXR/Source/OpenXRHMD/Public/OpenXRCore.h b/Engine/Plugins/Runtime/OpenXR/Source/OpenXRHMD/Public/OpenXRCore.h
--- a/Engine/Plugins/Runtime/OpenXR/Source/OpenXRHMD/Public/OpenXRCore.h
+++ b/Engine/Plugins/Runtime/OpenXR/Source/OpenXRHMD/Public/OpenXRCore.h
@@ -280,6 +280,8 @@ bool PreInitOpenXRCore(PFN_xrGetInstanceProcAddr InGetProcAddr);
  */
 bool InitOpenXRCore(XrInstance Instance);
 
+void OPENXRHMD_API EnumerateOpenXRApiLayers(TArray<XrApiLayerProperties>& OutProperties);
+
 FORCEINLINE void FilterActionName(const char* InActionName, char* OutActionName)
 {
 	static_assert(XR_MAX_ACTION_NAME_SIZE == XR_MAX_ACTION_SET_NAME_SIZE);
diff --git a/Engine/Plugins/Runtime/OpenXREyeTracker/Source/OpenXREyeTracker/Private/OpenXREyeTracker.cpp b/Engine/Plugins/Runtime/OpenXREyeTracker/Source/OpenXREyeTracker/Private/OpenXREyeTracker.cpp
--- a/Engine/Plugins/Runtime/OpenXREyeTracker/Source/OpenXREyeTracker/Private/OpenXREyeTracker.cpp
+++ b/Engine/Plugins/Runtime/OpenXREyeTracker/Source/OpenXREyeTracker/Private/OpenXREyeTracker.cpp
@@ -32,14 +32,7 @@ void FOpenXREyeTracker::Destroy()
 bool FOpenXREyeTracker::GetRequiredExtensions(TArray<const ANSICHAR*>& OutExtensions)
 {
 	TArray<XrApiLayerProperties> Properties;
-	uint32_t Count = 0;
-	XR_ENSURE(xrEnumerateApiLayerProperties(0, &Count, nullptr));
-	Properties.SetNum(Count);
-	for (auto& Prop : Properties)
-	{
-		Prop = XrApiLayerProperties{ XR_TYPE_API_LAYER_PROPERTIES };
-	}
-	XR_ENSURE(xrEnumerateApiLayerProperties(Count, &Count, Properties.GetData()));
+	EnumerateOpenXRApiLayers(Properties);
 
 	// Some API layers can crash the loader when enabled, if they're present we shouldn't enable the extension
 	for (const XrApiLayerProperties& Layer : Properties)
diff --git a/Engine/Plugins/Runtime/OpenXRHandTracking/Source/OpenXRHandTracking/Private/OpenXRHandTracking.cpp b/Engine/Plugins/Runtime/OpenXRHandTracking/Source/OpenXRHandTracking/Private/OpenXRHandTracking.cpp
--- a/Engine/Plugins/Runtime/OpenXRHandTracking/Source/OpenXRHandTracking/Private/OpenXRHandTracking.cpp
+++ b/Engine/Plugins/Runtime/OpenXRHandTracking/Source/OpenXRHandTracking/Private/OpenXRHandTracking.cpp
@@ -162,14 +162,7 @@ FOpenXRHandTracking::~FOpenXRHandTracking()
 bool FOpenXRHandTracking::GetRequiredExtensions(TArray<const ANSICHAR*>& OutExtensions)
 {
 	TArray<XrApiLayerProperties> Properties;
-	uint32_t Count = 0;
-	XR_ENSURE(xrEnumerateApiLayerProperties(0, &Count, nullptr));
-	Properties.SetNum(Count);
-	for (auto& Prop : Properties)
-	{
-		Prop = XrApiLayerProperties{ XR_TYPE_API_LAYER_PROPERTIES };
-	}
-	XR_ENSURE(xrEnumerateApiLayerProperties(Count, &Count, Properties.GetData()));
+	EnumerateOpenXRApiLayers(Properties);
 
 	// Some API layers can crash the loader when enabled, if they're present we shouldn't enable the extension
 	for (const XrApiLayerProperties& Layer : Properties)
