Migrate from using the deprecated PMIC to a list of IMCs

diff --git a/Engine/Plugins/Experimental/XRCreativeFramework/Source/XRCreativeEditor/Private/XRCreativeVREditorMode.cpp b/Engine/Plugins/Experimental/XRCreativeFramework/Source/XRCreativeEditor/Private/XRCreativeVREditorMode.cpp
--- a/Engine/Plugins/Experimental/XRCreativeFramework/Source/XRCreativeEditor/Private/XRCreativeVREditorMode.cpp
+++ b/Engine/Plugins/Experimental/XRCreativeFramework/Source/XRCreativeEditor/Private/XRCreativeVREditorMode.cpp
@@ -218,9 +218,15 @@ void UXRCreativeVREditorMode::EnableStereo()
 		Viewport->SetRenderDirectlyToWindow(true);
 	}
 
-	PRAGMA_DISABLE_DEPRECATION_WARNINGS
-	UOpenXRInputFunctionLibrary::BeginXRSession(Cast<UPlayerMappableInputConfig>(MappableInputConfig.TryLoad()));
-	PRAGMA_ENABLE_DEPRECATION_WARNINGS
+	PRAGMA_DISABLE_DEPRECATION_WARNINGS;
+	TObjectPtr<UPlayerMappableInputConfig> InputConfig = Cast<UPlayerMappableInputConfig>(MappableInputConfig.TryLoad());
+	TSet<UInputMappingContext*> Contexts;
+	for (const auto& Context : InputConfig->GetMappingContexts())
+	{
+		Contexts.Add(Context.Key);
+	}
+	PRAGMA_ENABLE_DEPRECATION_WARNINGS;
+	UOpenXRInputFunctionLibrary::BeginXRSession(Contexts);
 }
 
 
diff --git a/Engine/Plugins/Runtime/OpenXR/Source/OpenXRInput/Private/OpenXRInput.cpp b/Engine/Plugins/Runtime/OpenXR/Source/OpenXRInput/Private/OpenXRInput.cpp
--- a/Engine/Plugins/Runtime/OpenXR/Source/OpenXRInput/Private/OpenXRInput.cpp
+++ b/Engine/Plugins/Runtime/OpenXR/Source/OpenXRInput/Private/OpenXRInput.cpp
@@ -236,7 +236,7 @@ FOpenXRInputPlugin::FOpenXRInput::FOpenXRInput(FOpenXRHMD* HMD)
 	, EnhancedActions()
 	, Controllers()
 	, MotionSourceToControllerHandMap()
-	, MappableInputConfig(nullptr)
+	, InputMappingContexts()
 	, bActionsAttached(false)
 	, bDirectionalBindingSupported(false)
 	, bPalmPoseSupported(false)
@@ -375,17 +375,18 @@ bool FOpenXRInputPlugin::FOpenXRInput::BuildActions(XrSession Session)
 	}
 
 	PRAGMA_DISABLE_DEPRECATION_WARNINGS
-	if (!MappableInputConfig)
+	// Attempt to load the default input config from the OpenXR input settings.
+	UOpenXRInputSettings* InputSettings = GetMutableDefault<UOpenXRInputSettings>();
+	if (InputSettings)
 	{
-		// Attempt to load the default input config from the OpenXR input settings.
-		UOpenXRInputSettings* InputSettings = GetMutableDefault<UOpenXRInputSettings>();
-		if (InputSettings && InputSettings->MappableInputConfig.IsValid())
+		for (const auto& Context : InputSettings->InputMappingContexts)
 		{
-			SetPlayerMappableInputConfig((UPlayerMappableInputConfig*)InputSettings->MappableInputConfig.TryLoad());
+			TStrongObjectPtr<UInputMappingContext> Obj(Context.LoadSynchronous());
+			InputMappingContexts.Add(Obj);
 		}
 	}
 
-	if (MappableInputConfig)
+	if (!InputMappingContexts.IsEmpty())
 	{
 		BuildEnhancedActions(Profiles);
 	}
@@ -555,12 +556,12 @@ void FOpenXRInputPlugin::FOpenXRInput::BuildEnhancedActions(TMap<FString, FInter
 	}
 
 	PRAGMA_DISABLE_DEPRECATION_WARNINGS
-	for (const TPair<TObjectPtr<UInputMappingContext>, int32>& MappingContext : MappableInputConfig->GetMappingContexts())
+	for (const auto& MappingContext : InputMappingContexts)
 	{
-		FOpenXRActionSet ActionSet(Instance, MappingContext.Key->GetFName(), MappingContext.Key->ContextDescription.ToString(), ToXrPriority(MappingContext.Value), MappingContext.Key);
+		FOpenXRActionSet ActionSet(Instance, MappingContext->GetFName(), MappingContext->ContextDescription.ToString(), 0, MappingContext.Get());
 		TMap<FName, int32> ActionMap;
 
-		for (const FEnhancedActionKeyMapping& Mapping : MappingContext.Key->GetMappings())
+		for (const FEnhancedActionKeyMapping& Mapping : MappingContext->GetMappings())
 		{
 			if (!Mapping.Action)
 			{
@@ -744,7 +745,7 @@ void FOpenXRInputPlugin::FOpenXRInput::OnDestroySession()
 	{
 		// If the session shut down, clean up.
 		bActionsAttached = false;
-		MappableInputConfig = nullptr;
+		InputMappingContexts.Reset();
 	}
 }
 
@@ -1412,9 +1413,26 @@ bool FOpenXRInputPlugin::FOpenXRInput::SetPlayerMappableInputConfig(TObjectPtr<c
 		return false;
 	}
 
-	MappableInputConfig = TStrongObjectPtr<class UPlayerMappableInputConfig>(InputConfig);
-	return true;
+	TSet<TObjectPtr<UInputMappingContext>> MappingContexts;
+	InputConfig->GetMappingContexts().GetKeys(MappingContexts);
+	return AttachInputMappingContexts(MappingContexts);
 }
 PRAGMA_ENABLE_DEPRECATION_WARNINGS
 
+bool FOpenXRInputPlugin::FOpenXRInput::AttachInputMappingContexts(const TSet<TObjectPtr<UInputMappingContext>>& MappingContexts)
+{
+	if (bActionsAttached)
+	{
+		UE_LOG(LogHMD, Error, TEXT("Attempted to attach input mapping contexts when action sets are already attached for the current session."));
+
+		return false;
+	}
+
+	for (const auto& Context : MappingContexts)
+	{
+		InputMappingContexts.Add(TStrongObjectPtr<UInputMappingContext>(Context));
+	}
+	return true;
+}
+
 #undef LOCTEXT_NAMESPACE // "OpenXRInputPlugin"
diff --git a/Engine/Plugins/Runtime/OpenXR/Source/OpenXRInput/Private/OpenXRInput.h b/Engine/Plugins/Runtime/OpenXR/Source/OpenXRInput/Private/OpenXRInput.h
--- a/Engine/Plugins/Runtime/OpenXR/Source/OpenXRInput/Private/OpenXRInput.h
+++ b/Engine/Plugins/Runtime/OpenXR/Source/OpenXRInput/Private/OpenXRInput.h
@@ -130,6 +130,7 @@ public:
 		PRAGMA_DISABLE_DEPRECATION_WARNINGS
 		virtual bool SetPlayerMappableInputConfig(TObjectPtr<class UPlayerMappableInputConfig> InputConfig) override;
 		PRAGMA_ENABLE_DEPRECATION_WARNINGS
+		virtual bool AttachInputMappingContexts(const TSet<TObjectPtr<UInputMappingContext>>& MappingContexts) override;
 
 		// IHapticDevice overrides
 		IHapticDevice* GetHapticDevice() override { return (IHapticDevice*)this; }
@@ -149,9 +150,7 @@ public:
 		TMap<EControllerHand, FOpenXRController> Controllers;
 		TMap<FName, EControllerHand> MotionSourceToControllerHandMap;
 
-		PRAGMA_DISABLE_DEPRECATION_WARNINGS
-		TStrongObjectPtr<class UPlayerMappableInputConfig> MappableInputConfig;
-		PRAGMA_ENABLE_DEPRECATION_WARNINGS
+		TSet<TStrongObjectPtr<UInputMappingContext>> InputMappingContexts;
 
 		XrAction GetActionForMotionSource(FName MotionSource) const;
 		int32 GetDeviceIDForMotionSource(FName MotionSource) const;
diff --git a/Engine/Plugins/Runtime/OpenXR/Source/OpenXRInput/Private/OpenXRInputFunctionLibrary.cpp b/Engine/Plugins/Runtime/OpenXR/Source/OpenXRInput/Private/OpenXRInputFunctionLibrary.cpp
--- a/Engine/Plugins/Runtime/OpenXR/Source/OpenXRInput/Private/OpenXRInputFunctionLibrary.cpp
+++ b/Engine/Plugins/Runtime/OpenXR/Source/OpenXRInput/Private/OpenXRInputFunctionLibrary.cpp
@@ -16,17 +16,12 @@ UOpenXRInputFunctionLibrary::UOpenXRInputFunctionLibrary(const FObjectInitialize
 }
 
 PRAGMA_DISABLE_DEPRECATION_WARNINGS
-bool UOpenXRInputFunctionLibrary::BeginXRSession(UPlayerMappableInputConfig* InputConfig)
+bool UOpenXRInputFunctionLibrary::BeginXRSession(const TSet<UInputMappingContext*>& InputMappingContexts)
 {
-	if (!InputConfig)
+	TSet<TObjectPtr<UInputMappingContext>> Contexts;
+	for (UInputMappingContext* Context : InputMappingContexts)
 	{
-		// If the input config is unset, then we fall back to either the input config from the project settings.
-		// If the project settings is also unset, the input config will be set to null and we'll fall back to legacy input.
-		UOpenXRInputSettings* InputSettings = GetMutableDefault<UOpenXRInputSettings>();
-		if (InputSettings && InputSettings->MappableInputConfig.IsValid())
-		{
-			InputConfig = (UPlayerMappableInputConfig*)InputSettings->MappableInputConfig.TryLoad();
-		}
+		Contexts.Add(Context);
 	}
 
 	TArray<IMotionController*> MotionControllers = IModularFeatures::Get().GetModularFeatureImplementations<IMotionController>(IMotionController::GetModularFeatureName());
@@ -37,7 +32,7 @@ bool UOpenXRInputFunctionLibrary::BeginXRSession(UPlayerMappableInputConfig* Inp
 			continue;
 		}
 
-		MotionController->SetPlayerMappableInputConfig(InputConfig);
+		MotionController->AttachInputMappingContexts(Contexts);
 	}
 
 	return UHeadMountedDisplayFunctionLibrary::EnableHMD(true);
diff --git a/Engine/Plugins/Runtime/OpenXR/Source/OpenXRInput/Private/OpenXRInputSettings.cpp b/Engine/Plugins/Runtime/OpenXR/Source/OpenXRInput/Private/OpenXRInputSettings.cpp
--- a/Engine/Plugins/Runtime/OpenXR/Source/OpenXRInput/Private/OpenXRInputSettings.cpp
+++ b/Engine/Plugins/Runtime/OpenXR/Source/OpenXRInput/Private/OpenXRInputSettings.cpp
@@ -3,6 +3,7 @@
 #include "OpenXRInputSettings.h"
 #include "XRMotionControllerBase.h"
 #include "PlayerMappableInputConfig.h"
+#include "InputMappingContext.h"
 
 #if WITH_EDITOR
 #include "Editor.h"
@@ -15,26 +16,26 @@ UOpenXRInputSettings::UOpenXRInputSettings(const FObjectInitializer& ObjectIniti
 }
 
 #if WITH_EDITOR
-void UOpenXRInputSettings::PostEditChangeChainProperty(FPropertyChangedChainEvent& PropertyChangedEvent)
+PRAGMA_DISABLE_DEPRECATION_WARNINGS
+void UOpenXRInputSettings::PostInitProperties()
 {
-	Super::PostEditChangeChainProperty(PropertyChangedEvent);
+	Super::PostInitProperties();
 
-	const FName MemberPropertyName = PropertyChangedEvent.PropertyChain.GetActiveMemberNode()->GetValue()->GetFName();
-
-	PRAGMA_DISABLE_DEPRECATION_WARNINGS
-	if (MemberPropertyName == GET_MEMBER_NAME_CHECKED(UOpenXRInputSettings, MappableInputConfig))
+	if (MappableInputConfig.IsValid())
 	{
-		TArray<IMotionController*> MotionControllers = IModularFeatures::Get().GetModularFeatureImplementations<IMotionController>(IMotionController::GetModularFeatureName());
-		for (auto MotionController : MotionControllers)
+		UPlayerMappableInputConfig* InputConfig = Cast<UPlayerMappableInputConfig>(MappableInputConfig.TryLoad());
+		if (InputConfig)
 		{
-			if (MotionController == nullptr)
+			for (const auto& Context : InputConfig->GetMappingContexts())
 			{
-				continue;
+				TSoftObjectPtr<UInputMappingContext> Obj = Context.Key;
+				InputMappingContexts.Add(Obj);
 			}
+			MappableInputConfig.Reset();
 
-			MotionController->SetPlayerMappableInputConfig((UPlayerMappableInputConfig*)MappableInputConfig.TryLoad());
+			TryUpdateDefaultConfigFile();
 		}
 	}
-	PRAGMA_ENABLE_DEPRECATION_WARNINGS
 }
+PRAGMA_ENABLE_DEPRECATION_WARNINGS
 #endif
diff --git a/Engine/Plugins/Runtime/OpenXR/Source/OpenXRInput/Public/OpenXRInputFunctionLibrary.h b/Engine/Plugins/Runtime/OpenXR/Source/OpenXRInput/Public/OpenXRInputFunctionLibrary.h
--- a/Engine/Plugins/Runtime/OpenXR/Source/OpenXRInput/Public/OpenXRInputFunctionLibrary.h
+++ b/Engine/Plugins/Runtime/OpenXR/Source/OpenXRInput/Public/OpenXRInputFunctionLibrary.h
@@ -12,19 +12,16 @@ class OPENXRINPUT_API UOpenXRInputFunctionLibrary : public UBlueprintFunctionLib
 {
 	GENERATED_UCLASS_BODY()
 
-	PRAGMA_DISABLE_DEPRECATION_WARNINGS
 	/**
-	* Add a player mappable input config for Enhanced Input. This will make the tracking system aware of the actions and bindings that are used for XR motion controllers.
-	* Attaching input configs to the session can only be done once, so if multiple input configs need to be added only check AttachToSession for the last input config.
+	* Begin an XR session with the provided input mapping. This will make the tracking system aware of the actions and bindings that are used for XR motion controllers.
+	* Attaching input configs to the session can only be done once, so this is a helper function to attach the input mapping contexts and start the XR session correctly.
 	* 
-	* @param InputConfig		The path to the player mappable input config asset
-	* @param AttachToSession	Attach all pending configs to the running session
+	* @param InputMappingContexts		The set of input mapping contexts used for XR
 	*
-	* @return			False if the input config can't be attached to the session, true otherwise
+	* @return			False if the input mapping contexts can't be attached to the session, true otherwise
 	*/
 	UFUNCTION(BlueprintCallable, Category = "Input|XRTracking")
-	static bool BeginXRSession(class UPlayerMappableInputConfig* InputConfig);
-	PRAGMA_ENABLE_DEPRECATION_WARNINGS
+	static bool BeginXRSession(const TSet<UInputMappingContext*>& InputMappingContexts);
 
 
 	UFUNCTION(BlueprintCallable, Category = "Input|XRTracking")
diff --git a/Engine/Plugins/Runtime/OpenXR/Source/OpenXRInput/Public/OpenXRInputSettings.h b/Engine/Plugins/Runtime/OpenXR/Source/OpenXRInput/Public/OpenXRInputSettings.h
--- a/Engine/Plugins/Runtime/OpenXR/Source/OpenXRInput/Public/OpenXRInputSettings.h
+++ b/Engine/Plugins/Runtime/OpenXR/Source/OpenXRInput/Public/OpenXRInputSettings.h
@@ -17,12 +17,16 @@ class OPENXRINPUT_API UOpenXRInputSettings : public UObject
 public:
 	PRAGMA_DISABLE_DEPRECATION_WARNINGS
 	/** Set a mappable input config to allow OpenXR runtimes to remap the Enhanced Input actions. */
-	UPROPERTY(config, EditAnywhere, Category = "Enhanced Input", meta = (DisplayName = "Mappable Input Config for XR", AllowedClasses = "/Script/EnhancedInput.PlayerMappableInputConfig"))
+	UPROPERTY(config)
 	FSoftObjectPath MappableInputConfig = nullptr;
 	PRAGMA_ENABLE_DEPRECATION_WARNINGS
 
+	/** Set a mappable input config to allow OpenXR runtimes to remap the Enhanced Input actions. */
+	UPROPERTY(config, EditAnywhere, Category = "Enhanced Input", meta = (DisplayName = "Input Mapping Contexts for XR"))
+	TSet<TSoftObjectPtr<class UInputMappingContext>> InputMappingContexts;
+
 	// UObject interface
 #if WITH_EDITOR
-	virtual void PostEditChangeChainProperty(struct FPropertyChangedChainEvent& PropertyChangedEvent) override;
+	virtual void PostInitProperties() override;
 #endif
 };
diff --git a/Engine/Source/Runtime/HeadMountedDisplay/Public/IMotionController.h b/Engine/Source/Runtime/HeadMountedDisplay/Public/IMotionController.h
--- a/Engine/Source/Runtime/HeadMountedDisplay/Public/IMotionController.h
+++ b/Engine/Source/Runtime/HeadMountedDisplay/Public/IMotionController.h
@@ -160,9 +160,20 @@ public:
 	 *
 	 * @return			False if the input config can't be attached to the session, true otherwise
 	 */
+	UE_DEPRECATED(5.4, "UPlayerMappableInputConfig has been deprecated. Please use SetInputMappingContexts instead.")
 	virtual bool SetPlayerMappableInputConfig(TObjectPtr<class UPlayerMappableInputConfig> InputConfig = nullptr) { return true; };
 	PRAGMA_ENABLE_DEPRECATION_WARNINGS
 
+	/**
+	 * Add input mapping contexts to the motion controller. This allows the motion controller to support
+	 * Enhanced Input actions.
+	 *
+	 * @param MappingContexts		The set of input mapping contexts used with motion controllers
+	 *
+	 * @return			False if the input mapping contexts can't be attached to the session, true otherwise
+	 */
+	virtual bool AttachInputMappingContexts(const TSet<TObjectPtr<class UInputMappingContext>>& InputMappingContexts) { return true; };
+
 	// explicit source names
 	static HEADMOUNTEDDISPLAY_API FName LeftHandSourceId;
 	static HEADMOUNTEDDISPLAY_API FName RightHandSourceId;
diff --git a/Templates/TP_VirtualRealityBP/Config/DefaultInput.ini b/Templates/TP_VirtualRealityBP/Config/DefaultInput.ini
--- a/Templates/TP_VirtualRealityBP/Config/DefaultInput.ini
+++ b/Templates/TP_VirtualRealityBP/Config/DefaultInput.ini
@@ -108,6 +108,6 @@ DefaultTouchInterface=None
 +ConsoleKeys=Tilde
 
 [/Script/OpenXRInput.OpenXRInputSettings]
-MappableInputConfig=/Game/VRTemplate/Input/PMI_VRTemplate.PMI_VRTemplate
-bUseEnhancedInput=True
+MappableInputConfig=None
+InputMappingContexts=("/Game/VRTemplate/Input/IMC_Default.IMC_Default","/Game/VRTemplate/Input/IMC_Menu.IMC_Menu","/Game/VRTemplate/Input/IMC_Weapon_Left.IMC_Weapon_Left","/Game/VRTemplate/Input/IMC_Weapon_Right.IMC_Weapon_Right","/Game/VRTemplate/Input/IMC_Hands.IMC_Hands")
 
