Add support for changing foveation params at runtime + bug fixes

diff --git a/Engine/Plugins/Runtime/OpenXR/Source/OpenXRHMD/Private/FBFoveationImageGenerator.cpp b/Engine/Plugins/Runtime/OpenXR/Source/OpenXRHMD/Private/FBFoveationImageGenerator.cpp
--- a/Engine/Plugins/Runtime/OpenXR/Source/OpenXRHMD/Private/FBFoveationImageGenerator.cpp
+++ b/Engine/Plugins/Runtime/OpenXR/Source/OpenXRHMD/Private/FBFoveationImageGenerator.cpp
@@ -34,6 +34,7 @@ static TAutoConsoleVariable<int32> CVarFBFoveationPreview(
 
 FFBFoveationImageGenerator::FFBFoveationImageGenerator(bool bIsFoveationExtensionSupported, XrInstance InInstance, FOpenXRHMD* HMD, bool bMobileMultiViewEnabled)
 	: bIsMobileMultiViewEnabled(bMobileMultiViewEnabled)
+	, CurrentFrameSwapchainIndex(0)
 	, bFoveationExtensionSupported(bIsFoveationExtensionSupported)
 {
 	if (bFoveationExtensionSupported)
@@ -44,11 +45,10 @@ FFBFoveationImageGenerator::FFBFoveationImageGenerator(bool bIsFoveationExtensio
 
 		int32 SanitisedFoveationLevel = FMath::Clamp(CVarOpenXRFBFoveationLevel->GetInt(), XrFoveationLevelFB::XR_FOVEATION_LEVEL_NONE_FB, XrFoveationLevelFB::XR_FOVEATION_LEVEL_HIGH_FB);
 		bool bFoveationDynamic = CVarOpenXRFBFoveationDynamic->GetBool();
-		float SanitisedVerticalOffset = CVarOpenXRFBFoveationVerticalOffset->GetFloat();
-		SanitisedVerticalOffset = SanitisedVerticalOffset >= 0 ? SanitisedVerticalOffset : 0;
 
 		FoveationLevel = static_cast<XrFoveationLevelFB>(SanitisedFoveationLevel);
-		VerticalOffset = SanitisedVerticalOffset;
+		// Vertical offset is in degrees so it does not need to be clamped to values above 0.
+		VerticalOffset = CVarOpenXRFBFoveationVerticalOffset->GetFloat();
 		FoveationDynamic = bFoveationDynamic ? XR_FOVEATION_DYNAMIC_LEVEL_ENABLED_FB : XR_FOVEATION_DYNAMIC_DISABLED_FB;
 	}
 
@@ -57,7 +57,7 @@ FFBFoveationImageGenerator::FFBFoveationImageGenerator(bool bIsFoveationExtensio
 
 FRDGTextureRef FFBFoveationImageGenerator::GetImage(FRDGBuilder& GraphBuilder, const FViewInfo& ViewInfo, FVariableRateShadingImageManager::EVRSImageType ImageType)
 {
-	if (!bFoveationExtensionSupported || !OpenXRHMD || !FoveationImages.IsEmpty())
+	if (!bFoveationExtensionSupported || !OpenXRHMD || FoveationImages.IsEmpty())
 	{
 		return nullptr;
 	}
@@ -77,7 +77,9 @@ FRDGTextureRef FFBFoveationImageGenerator::GetImage(FRDGBuilder& GraphBuilder, c
 
 void FFBFoveationImageGenerator::PrepareImages(FRDGBuilder& GraphBuilder, const FSceneViewFamily& ViewFamily, const FMinimalSceneTextures& SceneTextures)
 {
-	return; //Currently not implemented as the images are prepared only when the color swapchain is reallocated.
+	//Not implemented as images are updated in UpdateFoveationImages only
+	//when foveation parameters change or when the color swapchain is reallocated.
+	return;
 }
 
 bool FFBFoveationImageGenerator::IsEnabled() const
@@ -101,12 +103,49 @@ FRDGTextureRef FFBFoveationImageGenerator::GetDebugImage(FRDGBuilder& GraphBuild
 	return nullptr;
 }
 
-void FFBFoveationImageGenerator::UpdateFoveationImages()
+void FFBFoveationImageGenerator::UpdateFoveationImages(bool bReallocatedSwapchain)
 {
 	if (!OpenXRHMD)
 	{
 		return;
 	}
+	int32 SanitisedFoveationLevel = FMath::Clamp(CVarOpenXRFBFoveationLevel->GetInt(), XrFoveationLevelFB::XR_FOVEATION_LEVEL_NONE_FB, XrFoveationLevelFB::XR_FOVEATION_LEVEL_HIGH_FB);
+	bool bFoveationDynamic = CVarOpenXRFBFoveationDynamic->GetBool();
+	float SanitisedVerticalOffset = CVarOpenXRFBFoveationVerticalOffset->GetFloat();
+
+	//Starting foveation level value outside actual range to force image creation at first update.
+	static int32 LastSanitisedFoveationLevel = -1;
+	static bool bLastFoveationDynamic = false;
+	static float LastSanitisedVerticalOffset = 0.0f;
+
+	bool bUpdateFoveationImages = false;
+	if (LastSanitisedFoveationLevel != SanitisedFoveationLevel)
+	{
+		LastSanitisedFoveationLevel = SanitisedFoveationLevel;
+		bUpdateFoveationImages = true;
+	}
+	if (bLastFoveationDynamic != bFoveationDynamic)
+	{
+		bLastFoveationDynamic = bFoveationDynamic;
+		bUpdateFoveationImages = true;
+	}
+	if (LastSanitisedVerticalOffset != SanitisedVerticalOffset)
+	{
+		LastSanitisedVerticalOffset = SanitisedVerticalOffset;
+		bUpdateFoveationImages = true;
+	}
+	
+	// If the swapchain has been reallocated we need to update the foveation images even if the foveation params haven't changed.
+	bUpdateFoveationImages |= bReallocatedSwapchain;
+
+	if (!bUpdateFoveationImages)
+	{
+		return;
+	}
+
+	FoveationLevel = static_cast<XrFoveationLevelFB>(SanitisedFoveationLevel);
+	VerticalOffset = SanitisedVerticalOffset;
+	FoveationDynamic = bFoveationDynamic ? XR_FOVEATION_DYNAMIC_LEVEL_ENABLED_FB : XR_FOVEATION_DYNAMIC_DISABLED_FB;
 
 	XrFoveationLevelProfileCreateInfoFB FoveationLevelProfileInfo{ XR_TYPE_FOVEATION_LEVEL_PROFILE_CREATE_INFO_FB };
 	FoveationLevelProfileInfo.next = nullptr;
diff --git a/Engine/Plugins/Runtime/OpenXR/Source/OpenXRHMD/Private/FBFoveationImageGenerator.h b/Engine/Plugins/Runtime/OpenXR/Source/OpenXRHMD/Private/FBFoveationImageGenerator.h
--- a/Engine/Plugins/Runtime/OpenXR/Source/OpenXRHMD/Private/FBFoveationImageGenerator.h
+++ b/Engine/Plugins/Runtime/OpenXR/Source/OpenXRHMD/Private/FBFoveationImageGenerator.h
@@ -17,7 +17,8 @@ public:
 	virtual ~FFBFoveationImageGenerator() override {};
 
 	bool IsFoveationExtensionEnabled() { return bFoveationExtensionSupported;}
-	void UpdateFoveationImages();
+	// bReallocatedSwapchain forces foveation images to update even if the foveation params haven't changed.
+	void UpdateFoveationImages(bool bReallocatedSwapchain = false);
 	void SetCurrentFrameSwapchainIndex(int32 CurrentFrameSwapchainIndex);
 
 	/** IVariableRateShadingImageGenerator interface */
diff --git a/Engine/Plugins/Runtime/OpenXR/Source/OpenXRHMD/Private/OpenXRHMD.cpp b/Engine/Plugins/Runtime/OpenXR/Source/OpenXRHMD/Private/OpenXRHMD.cpp
--- a/Engine/Plugins/Runtime/OpenXR/Source/OpenXRHMD/Private/OpenXRHMD.cpp
+++ b/Engine/Plugins/Runtime/OpenXR/Source/OpenXRHMD/Private/OpenXRHMD.cpp
@@ -2292,7 +2292,7 @@ bool FOpenXRHMD::AllocateRenderTargetTextures(uint32 SizeX, uint32 SizeY, uint8
 		}
 		if (FBFoveationImageGenerator && FBFoveationImageGenerator->IsFoveationExtensionEnabled())
 		{
-			FBFoveationImageGenerator->UpdateFoveationImages();
+			FBFoveationImageGenerator->UpdateFoveationImages(/* bReallocatedSwapchain */ true);
 		}
 	}
 
@@ -2846,6 +2846,7 @@ void FOpenXRHMD::OnBeginRendering_RenderThread(FRHICommandListImmediate& RHICmdL
 
 		if (bFoveationExtensionSupported && FBFoveationImageGenerator.IsValid())
 		{
+			FBFoveationImageGenerator->UpdateFoveationImages();
 			FBFoveationImageGenerator->SetCurrentFrameSwapchainIndex(ColorSwapchain->GetSwapChainIndex_RHIThread());
 		}
 
diff --git a/Engine/Plugins/Runtime/OpenXR/Source/OpenXRHMD/Private/OpenXRHMD_Swapchain.cpp b/Engine/Plugins/Runtime/OpenXR/Source/OpenXRHMD/Private/OpenXRHMD_Swapchain.cpp
--- a/Engine/Plugins/Runtime/OpenXR/Source/OpenXRHMD/Private/OpenXRHMD_Swapchain.cpp
+++ b/Engine/Plugins/Runtime/OpenXR/Source/OpenXRHMD/Private/OpenXRHMD_Swapchain.cpp
@@ -262,6 +262,11 @@ void FOpenXRSwapchain::GetFragmentDensityMaps(TArray<FTextureRHIRef>& OutTexture
 
 	TArray<XrSwapchainImageVulkanKHR> Images = EnumerateImages<XrSwapchainImageVulkanKHR>(Swapchain, XR_TYPE_SWAPCHAIN_IMAGE_VULKAN_KHR, Next);
 
+	if(!Images.IsEmpty())
+	{
+		OutTextureChain.Reset(Images.Num());
+	}
+
 	for (const XrSwapchainImageVulkanKHR& Image : Images)
 	{
 		const XrBaseOutStructure* NextHeader = reinterpret_cast<const XrBaseOutStructure*>(Image.next);
diff --git a/Engine/Source/Runtime/Renderer/Private/VariableRateShading/FoveatedImageGenerator.cpp b/Engine/Source/Runtime/Renderer/Private/VariableRateShading/FoveatedImageGenerator.cpp
--- a/Engine/Source/Runtime/Renderer/Private/VariableRateShading/FoveatedImageGenerator.cpp
+++ b/Engine/Source/Runtime/Renderer/Private/VariableRateShading/FoveatedImageGenerator.cpp
@@ -191,7 +191,7 @@ void FFoveatedImageGenerator::PrepareImages(FRDGBuilder& GraphBuilder, const FSc
 	}
 
 	// Sanity check VRS tile size.
-	check(GRHIVariableRateShadingImageTileMinWidth >= 8 && GRHIVariableRateShadingImageTileMinWidth <= 64 && GRHIVariableRateShadingImageTileMinHeight >= 8 && GRHIVariableRateShadingImageTileMaxHeight <= 64);
+	check(GRHIVariableRateShadingImageTileMinWidth >= 8 && GRHIVariableRateShadingImageTileMinWidth <= 64 && GRHIVariableRateShadingImageTileMinHeight >= 8 && GRHIVariableRateShadingImageTileMinHeight <= 64);
 
 	// Create texture to hold shading rate image
 	FRDGTextureDesc Desc = FVariableRateShadingImageManager::GetSRIDesc();
