Only attempt to update the config file in OpenXR settings if it is not read only or can be checked out successfully.

This prevents some logspam warnings about attempting to modify a read-only file/

#rnx

diff --git a/Engine/Plugins/Runtime/OpenXR/Source/OpenXRInput/OpenXRInput.Build.cs b/Engine/Plugins/Runtime/OpenXR/Source/OpenXRInput/OpenXRInput.Build.cs
--- a/Engine/Plugins/Runtime/OpenXR/Source/OpenXRInput/OpenXRInput.Build.cs
+++ b/Engine/Plugins/Runtime/OpenXR/Source/OpenXRInput/OpenXRInput.Build.cs
@@ -34,6 +34,7 @@ namespace UnrealBuildTool.Rules
 				PrivateDependencyModuleNames.Add("EditorFramework");
                 PrivateDependencyModuleNames.Add("UnrealEd");
 				PrivateDependencyModuleNames.Add("InputEditor");
+				PrivateDependencyModuleNames.Add("SourceControl");
 			}
         }
     }
diff --git a/Engine/Plugins/Runtime/OpenXR/Source/OpenXRInput/Private/OpenXRInputSettings.cpp b/Engine/Plugins/Runtime/OpenXR/Source/OpenXRInput/Private/OpenXRInputSettings.cpp
--- a/Engine/Plugins/Runtime/OpenXR/Source/OpenXRInput/Private/OpenXRInputSettings.cpp
+++ b/Engine/Plugins/Runtime/OpenXR/Source/OpenXRInput/Private/OpenXRInputSettings.cpp
@@ -4,10 +4,12 @@
 #include "XRMotionControllerBase.h"
 #include "PlayerMappableInputConfig.h"
 #include "InputMappingContext.h"
+#include "HAL/FileManager.h"
 
 #if WITH_EDITOR
 #include "Editor.h"
 #include "Features/IModularFeatures.h"
+#include "SourceControlHelpers.h"
 #endif
 
 UOpenXRInputSettings::UOpenXRInputSettings(const FObjectInitializer& ObjectInitializer)
@@ -33,7 +35,14 @@ void UOpenXRInputSettings::PostInitProperties()
 			}
 			MappableInputConfig.Reset();
 
-			TryUpdateDefaultConfigFile();
+			// Check if the file is read only. If it is, then try and check it out with source control
+			const FString ConfigFileName = GetDefaultConfigFilename();
+			bool bCanWriteToFile = !IFileManager::Get().IsReadOnly(*ConfigFileName) || (USourceControlHelpers::IsEnabled() && USourceControlHelpers::CheckOutFile(ConfigFileName));
+
+			if (bCanWriteToFile)
+			{
+				TryUpdateDefaultConfigFile();
+			}			
 		}
 	}
 }
