IOpenXRExtensionPlugin: Add a new function for adding suggested bindings
-Allow IOpenXRExtensionPlugins to add suggested bindings.

diff --git a/Engine/Plugins/Runtime/OpenXR/Source/OpenXRHMD/Public/IOpenXRExtensionPlugin.h b/Engine/Plugins/Runtime/OpenXR/Source/OpenXRHMD/Public/IOpenXRExtensionPlugin.h
--- a/Engine/Plugins/Runtime/OpenXR/Source/OpenXRHMD/Public/IOpenXRExtensionPlugin.h
+++ b/Engine/Plugins/Runtime/OpenXR/Source/OpenXRHMD/Public/IOpenXRExtensionPlugin.h
@@ -206,6 +206,16 @@ public:
 		return false;
 	}
 
+	/**
+	 * Set the output parameter to add suggested bindings to the given interaction profile.
+	 * This function gets called once for each interaction profile.
+	 * If false is returned the bindings will be ignored.
+	 */
+	virtual bool GetSuggestedBindings(XrPath InInteractionProfile, TArray<XrActionSuggestedBinding>& OutBindings)
+	{
+		return false;
+	}
+
 	/**
 	 * Set the output parameters to provide a path to an asset in the plugin content folder that visualizes
 	 * the controller in the hand represented by the user path.
diff --git a/Engine/Plugins/Runtime/OpenXR/Source/OpenXRInput/Private/OpenXRInput.cpp b/Engine/Plugins/Runtime/OpenXR/Source/OpenXRInput/Private/OpenXRInput.cpp
--- a/Engine/Plugins/Runtime/OpenXR/Source/OpenXRInput/Private/OpenXRInput.cpp
+++ b/Engine/Plugins/Runtime/OpenXR/Source/OpenXRInput/Private/OpenXRInput.cpp
@@ -406,6 +406,16 @@ bool FOpenXRInputPlugin::FOpenXRInput::BuildActions(XrSession Session)
 		// An exception is made for the Simple Controller Profile which is always bound as a fallback
 		if (Profile.Bindings.Num() > 0)
 		{
+			// Add bindings from the extension plugins
+			for (IOpenXRExtensionPlugin* Plugin : OpenXRHMD->GetExtensionPlugins())
+			{
+				TArray<XrActionSuggestedBinding> PluginBindings;
+				if (Plugin->GetSuggestedBindings(Profile.Path, PluginBindings))
+				{
+					Profile.Bindings.Append(PluginBindings);
+				}
+			}
+
 			// Add the bindings for the controller pose and haptics
 			Profile.Bindings.Add(XrActionSuggestedBinding{
 				Controllers[EControllerHand::Left].GripAction, FOpenXRPath("/user/hand/left/input/grip/pose")
