UE-198154 OpenXR LoadMap causes lag of presented frame vs hmd
-The two additional render thread pumps from RedrawViewports in LoadMap resulted in one xrWaitFrame on the game thread followed by three xrBeginFrame on the render thread.  This violates the OpenXR spec and put us into a bad state with the SteamVR runtime (and perhaps other runtimes).
-Fixed by adding a counter to the pipelined frame state so that we can tell if the xrWaitFrame that an xrBeginFrame is consuming has already been used, and early out.  A cvar exists to disable this behavior, xr.OpenXRUseWaitCountToAvoidExtraXrBeginFrameCalls, which is true by default can be set to false at startup to disable it.

diff --git a/Engine/Plugins/Runtime/OpenXR/Source/OpenXRHMD/Private/OpenXRHMD.cpp b/Engine/Plugins/Runtime/OpenXR/Source/OpenXRHMD/Private/OpenXRHMD.cpp
--- a/Engine/Plugins/Runtime/OpenXR/Source/OpenXRHMD/Private/OpenXRHMD.cpp
+++ b/Engine/Plugins/Runtime/OpenXR/Source/OpenXRHMD/Private/OpenXRHMD.cpp
@@ -108,6 +108,12 @@ static TAutoConsoleVariable<bool> CVarOpenXRAllowDepthLayer(
 	TEXT("Enables the depth composition layer if the XR_KHR_composition_layer_depth extension is supported.\n"),
 	ECVF_Default);
 
+static TAutoConsoleVariable<bool> CVarOpenXRUseWaitCountToAvoidExtraXrBeginFrameCalls(
+	TEXT("xr.OpenXRUseWaitCountToAvoidExtraXrBeginFrameCalls"),
+	true,
+	TEXT("If true we use the WaitCount in the PipelinedFrameState to avoid extra xrBeginFrame calls.  Without this level loads can cause two additional xrBeginFrame calls.\n"),
+	ECVF_Default);
+
 namespace {
 	static TSet<XrViewConfigurationType> SupportedViewConfigurations{ XR_VIEW_CONFIGURATION_TYPE_PRIMARY_STEREO, XR_VIEW_CONFIGURATION_TYPE_PRIMARY_QUAD_VARJO };
 
@@ -1328,6 +1334,7 @@ FOpenXRHMD::FOpenXRHMD(const FAutoRegister& AutoRegister, XrInstance InInstance,
 	bSupportsHandTracking = IsExtensionEnabled(XR_EXT_HAND_TRACKING_EXTENSION_NAME);
 	bSpaceAccelerationSupported = IsExtensionEnabled(XR_EPIC_SPACE_ACCELERATION_NAME);
 	bIsAcquireOnAnyThreadSupported = CheckPlatformAcquireOnAnyThreadSupport(InstanceProperties);
+	bUseWaitCountToAvoidExtraXrBeginFrameCalls = CVarOpenXRUseWaitCountToAvoidExtraXrBeginFrameCalls.GetValueOnAnyThread();
 	ReconfigureForShaderPlatform(GMaxRHIShaderPlatform);
 
 	bFoveationExtensionSupported = IsExtensionEnabled(XR_FB_SWAPCHAIN_UPDATE_STATE_EXTENSION_NAME) &&
@@ -2850,6 +2857,7 @@ void FOpenXRHMD::OnBeginRendering_RenderThread(FRHICommandListImmediate& RHICmdL
 			FBFoveationImageGenerator->SetCurrentFrameSwapchainIndex(ColorSwapchain->GetSwapChainIndex_RHIThread());
 		}
 
+		UE_LOG(LogHMD, VeryVerbose, TEXT("EnqueueLambda OnBeginRendering_RHIThread WaitCount: %i"), PipelinedFrameStateRendering.WaitCount);
 		RHICmdList.EnqueueLambda([this, FrameState = PipelinedFrameStateRendering, ColorSwapchain, DepthSwapchain, EmulationSwapchain](FRHICommandListImmediate& InRHICmdList)
 		{
 			OnBeginRendering_RHIThread(FrameState, ColorSwapchain, DepthSwapchain, EmulationSwapchain);
@@ -2940,6 +2948,7 @@ void FOpenXRHMD::OnBeginRendering_GameThread()
 			UE_CLOG(PipelinedFrameStateRendering.FrameState.predictedDisplayTime >= GameFrameState.FrameState.predictedDisplayTime,
 				LogHMD, VeryVerbose, TEXT("Predicted display time went backwards from %lld to %lld"), PipelinedFrameStateRendering.FrameState.predictedDisplayTime, GameFrameState.FrameState.predictedDisplayTime);
 
+			UE_LOG(LogHMD, VeryVerbose, TEXT("FOpenXRHMD TransferFrameStateToRenderingThread %i"), GameFrameState.WaitCount);
 			PipelinedFrameStateRendering = GameFrameState;
 
 			// Snapshot new poses for late update.
@@ -2987,12 +2996,17 @@ void FOpenXRHMD::OnBeginSimulation_GameThread()
 	{
 		FrameState.next = Module->OnWaitFrame(Session, FrameState.next);
 	}
+	static int WaitCount = 0;
+	++WaitCount;
+	UE_LOG(LogHMD, VeryVerbose, TEXT("xrWaitFrame %i"), WaitCount);
 	XR_ENSURE(xrWaitFrame(Session, &WaitInfo, &FrameState));
+	UE_LOG(LogHMD, VeryVerbose, TEXT("xrWaitFrame %i Complete"), WaitCount);
 
 	// The pipeline state on the game thread can only be safely modified after xrWaitFrame which will be unblocked by
 	// the runtime when xrBeginFrame is called. The rendering thread will clone the game pipeline state before calling
 	// xrBeginFrame so the game pipeline state can safely be modified after xrWaitFrame returns.
 
+	PipelineState.WaitCount = WaitCount;
 	PipelineState.bXrFrameStateUpdated = true;
 	PipelineState.FrameState = FrameState;
 	PipelineState.WorldToMetersScale = WorldToMetersScale;
@@ -3243,6 +3257,19 @@ void FOpenXRHMD::OnBeginRendering_RHIThread(const FPipelinedFrameState& InFrameS
 		return;
 	}
 
+	// We do not want xrBeginFrame to run twice based on a single xrWaitFrame.
+	// During LoadMap RedrawViewports(false) is called twice to pump the render thread without a new game thread pump.  This results in this function being
+	// called two additional times without corresponding xrWaitFrame calls from the game thread and therefore two extra xrBeginFrame calls.  On SteamVR, at least,
+	// this then leaves us in a situation where our xrWaitFrame immediately returns forever.
+	// To avoid this we ensure that each xrWaitFrame is consumed by xrBeginFrame only once.  We use the count of xrWaitFrame calls as an identifier.  Before 
+	// xrBeginFrame if the PipelinedFrameStateRHI wait count equals the incoming pipelined xrWaitFrame count then that xrWaitFrame has already been consumed,
+	// so we early out.  Once a new game frame happens and a new xrWaitFrame the early out will fail and xrBeginFrame will happen.
+	if ((PipelinedFrameStateRHI.WaitCount == InFrameState.WaitCount) && bUseWaitCountToAvoidExtraXrBeginFrameCalls)
+	{
+		UE_LOG(LogHMD, Verbose, TEXT("FOpenXRHMD::OnBeginRendering_RHIThread returning before xrBeginFrame because xrWaitFrame %i is already consumed.  This is expected twice during LoadMap and may also happen during other 'extra' render pumps."), InFrameState.WaitCount);
+		return;
+	}
+
 	// The layer state will be copied after SetFinalViewRect
 	PipelinedFrameStateRHI = InFrameState;
 
@@ -3254,7 +3281,9 @@ void FOpenXRHMD::OnBeginRendering_RHIThread(const FPipelinedFrameState& InFrameS
 	{
 		BeginInfo.next = Module->OnBeginFrame(Session, DisplayTime, BeginInfo.next);
 	}
-
+	static int BeginCount = 0;
+	PipelinedFrameStateRHI.BeginCount = ++BeginCount;
+	UE_LOG(LogHMD, VeryVerbose, TEXT("xrBeginFrame WaitCount: %i BeginCount: %i"), PipelinedFrameStateRHI.WaitCount, PipelinedFrameStateRHI.BeginCount);
 	XrResult Result = xrBeginFrame(Session, &BeginInfo);
 	if (XR_SUCCEEDED(Result))
 	{
@@ -3413,7 +3442,9 @@ void FOpenXRHMD::OnFinishRendering_RHIThread()
 		// thread. We have to do this per-frame because we can detach if app loses focus.
 		FAndroidApplication::GetJavaEnv();
 #endif
-
+		static int EndCount = 0;
+		PipelinedFrameStateRHI.EndCount = ++EndCount;
+		UE_LOG(LogHMD, VeryVerbose, TEXT("xrEndFrame WaitCount: %i BeginCount: %i EndCount: %i"), PipelinedFrameStateRHI.WaitCount, PipelinedFrameStateRHI.BeginCount, PipelinedFrameStateRHI.EndCount);
 		XR_ENSURE(xrEndFrame(Session, &EndInfo));
 	}
 
diff --git a/Engine/Plugins/Runtime/OpenXR/Source/OpenXRHMD/Private/OpenXRHMD.h b/Engine/Plugins/Runtime/OpenXR/Source/OpenXRHMD/Private/OpenXRHMD.h
--- a/Engine/Plugins/Runtime/OpenXR/Source/OpenXRHMD/Private/OpenXRHMD.h
+++ b/Engine/Plugins/Runtime/OpenXR/Source/OpenXRHMD/Private/OpenXRHMD.h
@@ -85,6 +85,9 @@ public:
 		TSharedPtr<FTrackingSpace> TrackingSpace;
 		float WorldToMetersScale = 100.0f;
 		float PixelDensity = 1.0f;
+		int WaitCount = 0;
+		int BeginCount = 0;
+		int EndCount = 0;
 		bool bXrFrameStateUpdated = false;
 	};
 
@@ -444,6 +447,7 @@ private:
 	bool					bIsStandaloneStereoOnlyDevice;
 	bool					bIsTrackingOnlySession;
 	bool					bIsAcquireOnAnyThreadSupported;
+	bool					bUseWaitCountToAvoidExtraXrBeginFrameCalls;
 	float					WorldToMetersScale = 100.0f;
 	float					RuntimePixelDensityMax = FHeadMountedDisplayBase::PixelDensityMax;
 	EShaderPlatform			ConfiguredShaderPlatform = EShaderPlatform::SP_NumPlatforms;
