UE-199135 GitHub 11029 : Destroy xrInstance in PreInit to yield other plugins to create xrInstance
-Destroy any openxr instance created during preinit becasue we do not yet know if openxrhmd will be the chosen xrsystem and another xrsystem may try to create instance (ie the oculus vr plugin) and then fail because the runtime only supports one instance.
-CVar xr.RetainPreInitInstance can be set to true to disable this behavior.

diff --git a/Engine/Plugins/Runtime/OpenXR/Source/OpenXRHMD/Private/OpenXRHMDModule.cpp b/Engine/Plugins/Runtime/OpenXR/Source/OpenXRHMD/Private/OpenXRHMDModule.cpp
--- a/Engine/Plugins/Runtime/OpenXR/Source/OpenXRHMD/Private/OpenXRHMDModule.cpp
+++ b/Engine/Plugins/Runtime/OpenXR/Source/OpenXRHMD/Private/OpenXRHMDModule.cpp
@@ -38,6 +38,12 @@ static TAutoConsoleVariable<bool> CVarCheckOpenXRInstanceConformance(
 	TEXT("If true, OpenXR will verify Instance is conformant by calling xrStringToPath. Some runtimes fail without a system attached at instance creation time."),
 	ECVF_RenderThreadSafe);
 
+static TAutoConsoleVariable<bool> CVarRetainPreInitInstance(
+	TEXT("xr.RetainPreInitInstance"),
+	false,
+	TEXT("If true, OpenXR will retain any instance created during PreInit rather than destroying it.  Destroying it is more correct because we are not yet certain to have chosen OpenXRHMD, and another HMD plugin could take over and try to create an instance of its own which would fail on a runtime that supports only one."),
+	ECVF_RenderThreadSafe);
+
 //---------------------------------------------------
 // OpenXRHMD Plugin Implementation
 //---------------------------------------------------
@@ -48,6 +54,8 @@ FOpenXRHMDModule::FOpenXRHMDModule()
 	: LoaderHandle(nullptr)
 	, Instance(XR_NULL_HANDLE)
 	, RenderBridge(nullptr)
+	, OculusAudioInputDevice()
+	, OculusAudioOutputDevice()
 { }
 
 FOpenXRHMDModule::~FOpenXRHMDModule()
@@ -87,6 +95,43 @@ TSharedPtr< class IXRTrackingSystem, ESPMode::ThreadSafe > FOpenXRHMDModule::Cre
 	return nullptr;
 }
 
+bool FOpenXRHMDModule::PreInit()
+{
+#if PLATFORM_WINDOWS
+	// On Windows, we need to get the audio input/output devices before init, so create the instance first, grab the audio devices, and then
+	// immediately destroy it so a new one can be created for the actual initialize call
+	const bool bInitialized = InitInstance();
+	if (bInitialized)
+	{
+		if (IsExtensionEnabled(XR_OCULUS_AUDIO_DEVICE_GUID_EXTENSION_NAME))
+		{
+			PFN_xrGetAudioInputDeviceGuidOculus GetAudioInputDeviceGuidOculus = nullptr;
+			if (XR_ENSURE(xrGetInstanceProcAddr(Instance, "xrGetAudioInputDeviceGuidOculus", (PFN_xrVoidFunction*)&GetAudioInputDeviceGuidOculus)))
+			{
+				WCHAR DeviceGuid[XR_MAX_AUDIO_DEVICE_STR_SIZE_OCULUS];
+				GetAudioInputDeviceGuidOculus(Instance, DeviceGuid);
+				OculusAudioInputDevice = FString(XR_MAX_AUDIO_DEVICE_STR_SIZE_OCULUS, DeviceGuid);
+			}
+
+			PFN_xrGetAudioOutputDeviceGuidOculus GetAudioOutputDeviceGuidOculus = nullptr;
+			if (XR_ENSURE(xrGetInstanceProcAddr(Instance, "xrGetAudioOutputDeviceGuidOculus", (PFN_xrVoidFunction*)&GetAudioOutputDeviceGuidOculus)))
+			{
+				WCHAR DeviceGuid[XR_MAX_AUDIO_DEVICE_STR_SIZE_OCULUS];
+				GetAudioOutputDeviceGuidOculus(Instance, DeviceGuid);
+				OculusAudioOutputDevice = FString(XR_MAX_AUDIO_DEVICE_STR_SIZE_OCULUS, DeviceGuid);
+			}
+		}
+		if (!CVarRetainPreInitInstance.GetValueOnAnyThread())
+		{
+			XR_ENSURE(xrDestroyInstance(Instance));
+			Instance = nullptr;
+		}
+	}
+	return bInitialized;
+#endif // PLATFORM_WINDOWS
+	return true;
+}
+
 void FOpenXRHMDModule::ShutdownModule()
 {
 	if (Instance)
@@ -1134,34 +1179,10 @@ XrPath FOpenXRHMDModule::ResolveNameToPath(FName Name)
 
 FString FOpenXRHMDModule::GetAudioInputDevice()
 {
-#if PLATFORM_WINDOWS
-	if (Instance && IsExtensionEnabled(XR_OCULUS_AUDIO_DEVICE_GUID_EXTENSION_NAME))
-	{
-		PFN_xrGetAudioInputDeviceGuidOculus GetAudioInputDeviceGuidOculus = nullptr;
-		if (XR_ENSURE(xrGetInstanceProcAddr(Instance, "xrGetAudioInputDeviceGuidOculus", (PFN_xrVoidFunction*)&GetAudioInputDeviceGuidOculus)))
-		{
-			WCHAR DeviceGuid[XR_MAX_AUDIO_DEVICE_STR_SIZE_OCULUS];
-			GetAudioInputDeviceGuidOculus(Instance, DeviceGuid);
-			return FString(XR_MAX_AUDIO_DEVICE_STR_SIZE_OCULUS, DeviceGuid);
-		}
-	}
-#endif // PLATFORM_WINDOWS
-	return FString();
+	return OculusAudioInputDevice;
 }
 
 FString FOpenXRHMDModule::GetAudioOutputDevice()
 {
-#if PLATFORM_WINDOWS
-	if (Instance && IsExtensionEnabled(XR_OCULUS_AUDIO_DEVICE_GUID_EXTENSION_NAME))
-	{
-		PFN_xrGetAudioOutputDeviceGuidOculus GetAudioOutputDeviceGuidOculus = nullptr;
-		if (XR_ENSURE(xrGetInstanceProcAddr(Instance, "xrGetAudioOutputDeviceGuidOculus", (PFN_xrVoidFunction*)&GetAudioOutputDeviceGuidOculus)))
-		{
-			WCHAR DeviceGuid[XR_MAX_AUDIO_DEVICE_STR_SIZE_OCULUS];
-			GetAudioOutputDeviceGuidOculus(Instance, DeviceGuid);
-			return FString(XR_MAX_AUDIO_DEVICE_STR_SIZE_OCULUS, DeviceGuid);
-		}
-	}
-#endif // PLATFORM_WINDOWS
-	return FString();
+	return OculusAudioOutputDevice;
 }
diff --git a/Engine/Plugins/Runtime/OpenXR/Source/OpenXRHMD/Private/OpenXRHMDModule.h b/Engine/Plugins/Runtime/OpenXR/Source/OpenXRHMD/Private/OpenXRHMDModule.h
--- a/Engine/Plugins/Runtime/OpenXR/Source/OpenXRHMD/Private/OpenXRHMDModule.h
+++ b/Engine/Plugins/Runtime/OpenXR/Source/OpenXRHMD/Private/OpenXRHMDModule.h
@@ -31,7 +31,7 @@ public:
 		AliasesOut.Add(TEXT("OpenXR"));
 	}
 
-	virtual bool PreInit() override { return InitInstance(); }
+	virtual bool PreInit() override; 
 	virtual void ShutdownModule() override;
 
 	virtual bool IsHMDConnected() override { return true; }
@@ -65,6 +65,10 @@ private:
 	TSortedMap<XrPath, FName> PathToName;
 	TSortedMap<FName, XrPath, FDefaultAllocator, FNameFastLess> NameToPath;
 
+	// Cache off Oculus Audio devices on PreInit so that the XrInstance can be released before Initialize
+	FString OculusAudioInputDevice;
+	FString OculusAudioOutputDevice;
+
 	bool EnumerateExtensions();
 	bool EnumerateLayers();
 	bool InitRenderBridge();
