PR #10695: Fix to allow headless XR Tracking Only mode to provide motion controller data

diff --git a/Engine/Plugins/Runtime/OpenXR/Source/OpenXRHMD/Private/OpenXRHMD.cpp b/Engine/Plugins/Runtime/OpenXR/Source/OpenXRHMD/Private/OpenXRHMD.cpp
--- a/Engine/Plugins/Runtime/OpenXR/Source/OpenXRHMD/Private/OpenXRHMD.cpp
+++ b/Engine/Plugins/Runtime/OpenXR/Source/OpenXRHMD/Private/OpenXRHMD.cpp
@@ -3252,7 +3252,7 @@ void FOpenXRHMD::OnBeginRendering_RHIThread(const FPipelinedFrameState& InFrameS
 	SCOPED_NAMED_EVENT(BeginFrame, FColor::Red);
 
 	FReadScopeLock Lock(SessionHandleMutex);
-	if (!bIsRunning || !RenderBridge)
+	if (!bIsRunning || (!RenderBridge && !bIsTrackingOnlySession))
 	{
 		return;
 	}
