Improve extensibility of IOpenXRExtensionPlugin by adding overrideable callbacks: OnSuggestBindings, OnCreateActionSet, PostCreateActionSet, OnCreateAction, PostCreateAction, and OnActionSetAttach.

diff --git a/Engine/Plugins/Runtime/OpenXR/Source/OpenXRHMD/Public/IOpenXRExtensionPlugin.h b/Engine/Plugins/Runtime/OpenXR/Source/OpenXRHMD/Public/IOpenXRExtensionPlugin.h
--- a/Engine/Plugins/Runtime/OpenXR/Source/OpenXRHMD/Public/IOpenXRExtensionPlugin.h
+++ b/Engine/Plugins/Runtime/OpenXR/Source/OpenXRHMD/Public/IOpenXRExtensionPlugin.h
@@ -391,12 +391,46 @@ public:
 		return InNext;
 	}
 
+	// FOpenXRInputPlugin::FOpenXRInput::BuildActions
+	virtual const void* OnSuggestBindings(XrPath InteractionProfile, const void* InNext)
+	{
+		return InNext;
+	}
+
 	// FOpenXRRenderBridge::Present, RHI thread
 	virtual const void* OnEndFrame(XrSession InSession, XrTime DisplayTime, const void* InNext)
 	{
 		return InNext;
 	}
 
+	// FOpenXRInputPlugin::FOpenXRActionSet::FOpenXRActionSet
+	virtual const void* OnCreateActionSet(XrActionSetCreateInfo InCreateInfo, const void* InNext)
+	{
+		return InNext;
+	}
+
+	// FOpenXRInputPlugin::FOpenXRActionSet::FOpenXRActionSet
+	void PostCreateActionSet(XrActionSet InActionSet)
+	{
+	}
+
+	// FOpenXRInputPlugin::FOpenXRAction::FOpenXRAction
+	virtual const void* OnCreateAction(XrActionCreateInfo InCreateInfo, const void* InNext)
+	{
+		return InNext;
+	}
+
+	// FOpenXRInputPlugin::FOpenXRAction::FOpenXRAction
+	void PostCreateAction(XrAction InAction)
+	{
+	}
+
+	// FOpenXRInputPlugin::FOpenXRInput::BuildActions
+	virtual const void* OnActionSetAttach(XrSessionActionSetsAttachInfo InAttachInfo, const void* InNext)
+	{
+		return InNext;
+	}
+
 	// FOpenXRInput::Tick, game thread, setting up for xrSyncActions.  This happens near the start of the game frame.
 	virtual const void* OnSyncActions(XrSession InSession, const void* InNext)
 	{
diff --git a/Engine/Plugins/Runtime/OpenXR/Source/OpenXRInput/Private/OpenXRInput.cpp b/Engine/Plugins/Runtime/OpenXR/Source/OpenXRInput/Private/OpenXRInput.cpp
--- a/Engine/Plugins/Runtime/OpenXR/Source/OpenXRInput/Private/OpenXRInput.cpp
+++ b/Engine/Plugins/Runtime/OpenXR/Source/OpenXRInput/Private/OpenXRInput.cpp
@@ -96,15 +96,15 @@ void FOpenXRInputPlugin::StartupModule()
 
 FOpenXRInputPlugin::FOpenXRAction::FOpenXRAction(XrActionSet InActionSet,
 	XrActionType InActionType, const FName& InName, const FString& InLocalizedName,
-	const TArray<XrPath>& InSubactionPaths, const TObjectPtr<const UInputAction>& InObject)
-	: FOpenXRAction(InActionSet, InActionType, InName, InLocalizedName, InSubactionPaths)
+	const TArray<XrPath>& InSubactionPaths, const TObjectPtr<const UInputAction>& InObject, FOpenXRHMD* OpenXRHMD)
+	: FOpenXRAction(InActionSet, InActionType, InName, InLocalizedName, InSubactionPaths, OpenXRHMD)
 {
 	Object = InObject;
 }
 
 FOpenXRInputPlugin::FOpenXRAction::FOpenXRAction(XrActionSet InActionSet,
 	XrActionType InActionType, const FName& InName, const FString& InLocalizedName,
-	const TArray<XrPath>& InSubactionPaths)
+	const TArray<XrPath>& InSubactionPaths, FOpenXRHMD* OpenXRHMD)
 	: Set(InActionSet)
 	, Type(InActionType)
 	, Name(InName)
@@ -130,19 +130,35 @@ FOpenXRInputPlugin::FOpenXRAction::FOpenXRAction(XrActionSet InActionSet,
 		FCStringAnsi::Strcpy(Info.localizedActionName, XR_MAX_LOCALIZED_ACTION_NAME_SIZE, ActionName);
 	}
 
+	if (OpenXRHMD)
+	{
+		for (IOpenXRExtensionPlugin* Plugin : OpenXRHMD->GetExtensionPlugins())
+		{
+			Info.next = Plugin->OnCreateAction(Info, Info.next);
+		}
+	}
+
 	XR_ENSURE(xrCreateAction(Set, &Info, &Handle));
+
+	if (OpenXRHMD)
+	{
+		for (IOpenXRExtensionPlugin* Plugin : OpenXRHMD->GetExtensionPlugins())
+		{
+			Plugin->PostCreateAction(Handle);
+		}
+	}
 }
 
 FOpenXRInputPlugin::FOpenXRActionSet::FOpenXRActionSet(XrInstance InInstance,
 	const FName& InName, const FString& InLocalizedName, uint32 InPriority,
-	const TObjectPtr<const UInputMappingContext>& InObject)
-	: FOpenXRActionSet(InInstance, InName, InLocalizedName, InPriority)
+	const TObjectPtr<const UInputMappingContext>& InObject, FOpenXRHMD* OpenXRHMD)
+	: FOpenXRActionSet(InInstance, InName, InLocalizedName, InPriority, OpenXRHMD)
 {
 	Object = InObject;
 }
 
 FOpenXRInputPlugin::FOpenXRActionSet::FOpenXRActionSet(XrInstance InInstance,
-	const FName& InName, const FString& InLocalizedName, uint32 InPriority)
+	const FName& InName, const FString& InLocalizedName, uint32 InPriority, FOpenXRHMD* OpenXRHMD)
 	: Handle(XR_NULL_HANDLE)
 	, Name(InName)
 	, LocalizedName(InLocalizedName)
@@ -165,7 +181,25 @@ FOpenXRInputPlugin::FOpenXRActionSet::FOpenXRActionSet(XrInstance InInstance,
 		FCStringAnsi::Strcpy(Info.localizedActionSetName, XR_MAX_LOCALIZED_ACTION_SET_NAME_SIZE, ActionName);
 	}
 	Info.priority = InPriority;
+
+	if (OpenXRHMD)
+	{
+		for (IOpenXRExtensionPlugin* Plugin : OpenXRHMD->GetExtensionPlugins())
+		{
+			Info.next = Plugin->OnCreateActionSet(Info, Info.next);
+		}
+	}
+
 	XR_ENSURE(xrCreateActionSet(InInstance, &Info, &Handle));
+
+	if (OpenXRHMD)
+	{
+		for (IOpenXRExtensionPlugin* Plugin : OpenXRHMD->GetExtensionPlugins())
+		{
+			Plugin->PostCreateActionSet(Handle);
+		}
+	}
+
 }
 
 FOpenXRInputPlugin::FOpenXRController::FOpenXRController(XrActionSet InActionSet, XrPath InUserPath, const char* InName)
@@ -466,6 +500,11 @@ bool FOpenXRInputPlugin::FOpenXRInput::BuildActions(XrSession Session)
 			InteractionProfile.countSuggestedBindings = Profile.Bindings.Num();
 			InteractionProfile.suggestedBindings = Profile.Bindings.GetData();
 
+			for (IOpenXRExtensionPlugin* Plugin : OpenXRHMD->GetExtensionPlugins())
+			{
+				InteractionProfile.next = Plugin->OnSuggestBindings(Profile.Path, InteractionProfile.next);
+			}
+
 			XR_ENSURE(xrSuggestInteractionProfileBindings(Instance, &InteractionProfile));
 		}
 	}
@@ -505,6 +544,12 @@ bool FOpenXRInputPlugin::FOpenXRInput::BuildActions(XrSession Session)
 	SessionActionSetsAttachInfo.next = nullptr;
 	SessionActionSetsAttachInfo.countActionSets = AttachArray.Num();
 	SessionActionSetsAttachInfo.actionSets = AttachArray.GetData();
+
+	for (IOpenXRExtensionPlugin* Plugin : OpenXRHMD->GetExtensionPlugins())
+	{
+		SessionActionSetsAttachInfo.next = Plugin->OnActionSetAttach(SessionActionSetsAttachInfo, SessionActionSetsAttachInfo.next);
+	}
+
 	bActionsAttached = XR_ENSURE(xrAttachSessionActionSets(Session, &SessionActionSetsAttachInfo));
 
 	return bActionsAttached;
diff --git a/Engine/Plugins/Runtime/OpenXR/Source/OpenXRInput/Private/OpenXRInput.h b/Engine/Plugins/Runtime/OpenXR/Source/OpenXRInput/Private/OpenXRInput.h
--- a/Engine/Plugins/Runtime/OpenXR/Source/OpenXRInput/Private/OpenXRInput.h
+++ b/Engine/Plugins/Runtime/OpenXR/Source/OpenXRInput/Private/OpenXRInput.h
@@ -43,13 +43,15 @@ public:
 			const FName& InName, const
 			FString& InLocalizedName,
 			const TArray<XrPath>& InSubactionPaths,
-			const TObjectPtr<const UInputAction>& InObject);
+			const TObjectPtr<const UInputAction>& InObject,
+			FOpenXRHMD* OpenXRHMD = nullptr);
 
 		FOpenXRAction(XrActionSet InActionSet,
 			XrActionType InActionType,
 			const FName& InName,
 			const FString& InLocalizedName,
-			const TArray<XrPath>& InSubactionPaths);
+			const TArray<XrPath>& InSubactionPaths,
+			FOpenXRHMD* OpenXRHMD = nullptr);
 	};
 
 	struct FOpenXRActionSet
@@ -64,12 +66,14 @@ public:
 			const FName& InName,
 			const FString& InLocalizedName,
 			uint32 InPriority,
-			const TObjectPtr<const UInputMappingContext>& InObject);
+			const TObjectPtr<const UInputMappingContext>& InObject,
+			FOpenXRHMD* OpenXRHMD = nullptr);
 
 		FOpenXRActionSet(XrInstance InInstance,
 			const FName& InName,
 			const FString& InLocalizedName,
-			uint32 InPriority);
+			uint32 InPriority,
+			FOpenXRHMD* OpenXRHMD = nullptr);
 	};
 
 	struct FOpenXRController
