UE-164875 Assert due to NaN passed in from Oculus OpenXR runtime with valid bit set
-Added an additional openxr function call failure handling branch.  The runtime should be setting the flags invalid if it entirely fails, but at least one runtime in the past did not do that so we may as well handle call failure eplicitly.

diff --git a/Engine/Plugins/Runtime/OpenXR/Source/OpenXRHMD/Private/OpenXRHMD.cpp b/Engine/Plugins/Runtime/OpenXR/Source/OpenXRHMD/Private/OpenXRHMD.cpp
--- a/Engine/Plugins/Runtime/OpenXR/Source/OpenXRHMD/Private/OpenXRHMD.cpp
+++ b/Engine/Plugins/Runtime/OpenXR/Source/OpenXRHMD/Private/OpenXRHMD.cpp
@@ -1472,22 +1472,25 @@ void FOpenXRHMD::UpdateDeviceLocations(bool bUpdateOpenXRExtensionPlugins)
 					// The display time is no longer valid so set the location as invalid as well
 					PipelineState.DeviceLocations[DeviceIndex].locationFlags = 0;
 				}
-				else
+				else if (Result != XR_SUCCESS)
 				{
-					XR_ENSURE(Result);
-				}
-				
-				// Clear the location tracked bits
-				CachedDeviceLocation.locationFlags &= ~(XR_SPACE_LOCATION_POSITION_TRACKED_BIT | XR_SPACE_LOCATION_ORIENTATION_TRACKED_BIT);
-				if (NewDeviceLocation.locationFlags & (XR_SPACE_LOCATION_POSITION_VALID_BIT))
-				{
-					CachedDeviceLocation.pose.position = NewDeviceLocation.pose.position;
-					CachedDeviceLocation.locationFlags |= (NewDeviceLocation.locationFlags & (XR_SPACE_LOCATION_POSITION_TRACKED_BIT | XR_SPACE_LOCATION_POSITION_VALID_BIT));
+					PipelineState.DeviceLocations[DeviceIndex].locationFlags = 0;
+					ensureMsgf(XR_SUCCEEDED(Result), TEXT("OpenXR xrLocateSpace failed with result %s.  No pose fetched."), OpenXRResultToString(Result)); \
 				}
-				if (NewDeviceLocation.locationFlags & (XR_SPACE_LOCATION_ORIENTATION_VALID_BIT))
+				else
 				{
-					CachedDeviceLocation.pose.orientation = NewDeviceLocation.pose.orientation;
-					CachedDeviceLocation.locationFlags |= (NewDeviceLocation.locationFlags & (XR_SPACE_LOCATION_ORIENTATION_TRACKED_BIT | XR_SPACE_LOCATION_ORIENTATION_VALID_BIT));
+					// Clear the location tracked bits
+					CachedDeviceLocation.locationFlags &= ~(XR_SPACE_LOCATION_POSITION_TRACKED_BIT | XR_SPACE_LOCATION_ORIENTATION_TRACKED_BIT);
+					if (NewDeviceLocation.locationFlags & (XR_SPACE_LOCATION_POSITION_VALID_BIT))
+					{
+						CachedDeviceLocation.pose.position = NewDeviceLocation.pose.position;
+						CachedDeviceLocation.locationFlags |= (NewDeviceLocation.locationFlags & (XR_SPACE_LOCATION_POSITION_TRACKED_BIT | XR_SPACE_LOCATION_POSITION_VALID_BIT));
+					}
+					if (NewDeviceLocation.locationFlags & (XR_SPACE_LOCATION_ORIENTATION_VALID_BIT))
+					{
+						CachedDeviceLocation.pose.orientation = NewDeviceLocation.pose.orientation;
+						CachedDeviceLocation.locationFlags |= (NewDeviceLocation.locationFlags & (XR_SPACE_LOCATION_ORIENTATION_TRACKED_BIT | XR_SPACE_LOCATION_ORIENTATION_VALID_BIT));
+					}
 				}
 			}
 			else
