UE-193724 OpenXRViveTracker references the Private folder of the OpenXRHMD module
-Created IOpenXRHMD interface so that other plugins can do OpenXR specific things with the OpenXRHMD implementation, rather than accessing FOpenXRHMD.

diff --git a/Engine/Plugins/Runtime/OpenXR/Source/OpenXRHMD/Private/OpenXRHMD.h b/Engine/Plugins/Runtime/OpenXR/Source/OpenXRHMD/Private/OpenXRHMD.h
--- a/Engine/Plugins/Runtime/OpenXR/Source/OpenXRHMD/Private/OpenXRHMD.h
+++ b/Engine/Plugins/Runtime/OpenXR/Source/OpenXRHMD/Private/OpenXRHMD.h
@@ -15,6 +15,7 @@
 #include "DefaultSpectatorScreenController.h"
 #include "IHeadMountedDisplayVulkanExtensions.h"
 #include "IOpenXRExtensionPluginDelegates.h"
+#include "IOpenXRHMD.h"
 #include "Misc/EnumClassFlags.h"
 
 #include <openxr/openxr.h>
@@ -40,6 +41,7 @@ class FOpenXRHMD
 	, public FOpenXRAssetManager
 	, public TStereoLayerManager<FOpenXRLayer>
 	, public IOpenXRExtensionPluginDelegates
+	, public IOpenXRHMD
 {
 private:
 
@@ -215,7 +217,7 @@ public:
 
 	virtual bool GetIsTracked(int32 DeviceId);
 	virtual bool GetCurrentPose(int32 DeviceId, FQuat& CurrentOrientation, FVector& CurrentPosition) override;
-	virtual bool GetPoseForTime(int32 DeviceId, FTimespan Timespan, bool& OutTimeWasUsed, FQuat& CurrentOrientation, FVector& CurrentPosition, bool& bProvidedLinearVelocity, FVector& LinearVelocity, bool& bProvidedAngularVelocity, FVector& AngularVelocityAsAxisAndLength, bool& bProvidedLinearAcceleration, FVector& LinearAcceleration, float WorldToMetersScale);
+	virtual bool GetPoseForTime(int32 DeviceId, FTimespan Timespan, bool& OutTimeWasUsed, FQuat& CurrentOrientation, FVector& CurrentPosition, bool& bProvidedLinearVelocity, FVector& LinearVelocity, bool& bProvidedAngularVelocity, FVector& AngularVelocityAsAxisAndLength, bool& bProvidedLinearAcceleration, FVector& LinearAcceleration, float WorldToMetersScale) override;
 	virtual bool GetCurrentInteractionProfile(const EControllerHand Hand, FString& InteractionProfile) override;
 	
 	virtual void SetBaseRotation(const FRotator& InBaseRotation) override;
@@ -311,6 +313,7 @@ public:
 	virtual bool DoesSupportLateProjectionUpdate() const override { return true; }
 	virtual FString GetVersionString() const override;
 	virtual bool HasValidTrackingPosition() override { return IsTracking(HMDDeviceId); }
+	virtual IOpenXRHMD* GetIOpenXRHMD() { return this; }
 
 	/** IHeadMountedDisplay interface */
 	virtual bool IsHMDConnected() override;
@@ -406,22 +409,23 @@ public:
 	void OnBeginRendering_RHIThread(const FPipelinedFrameState& InFrameState, FXRSwapChainPtr ColorSwapchain, FXRSwapChainPtr DepthSwapchain, FXRSwapChainPtr EmulationSwapchain);
 	void OnFinishRendering_RHIThread();
 
+	/** IOpenXRHMD */
 	/** @return	True if the HMD was initialized OK */
-	OPENXRHMD_API bool IsInitialized() const;
-	OPENXRHMD_API bool IsRunning() const;
-	OPENXRHMD_API bool IsFocused() const;
-
-	OPENXRHMD_API int32 AddTrackedDevice(XrAction Action, XrPath Path);
-	OPENXRHMD_API void ResetTrackedDevices();
-	OPENXRHMD_API XrPath GetTrackedDevicePath(const int32 DeviceId);
-	OPENXRHMD_API XrSpace GetTrackedDeviceSpace(const int32 DeviceId);
-
-	OPENXRHMD_API bool IsExtensionEnabled(const FString& Name) const { return EnabledExtensions.Contains(Name); }
-	OPENXRHMD_API XrInstance GetInstance() { return Instance; }
-	OPENXRHMD_API XrSystemId GetSystem() { return System; }
-	OPENXRHMD_API XrSession GetSession() { return Session; }
-	OPENXRHMD_API XrTime GetDisplayTime() const;
-	OPENXRHMD_API XrSpace GetTrackingSpace() const;
+	bool IsInitialized() const override;
+	bool IsRunning() const override;
+	bool IsFocused() const override;
+	int32 AddTrackedDevice(XrAction Action, XrPath Path) override;
+	void ResetTrackedDevices() override;
+	XrPath GetTrackedDevicePath(const int32 DeviceId) override;
+	XrSpace GetTrackedDeviceSpace(const int32 DeviceId) override;
+	bool IsExtensionEnabled(const FString& Name) const override { return EnabledExtensions.Contains(Name); }
+	XrInstance GetInstance() override { return Instance; }
+	XrSystemId GetSystem() override { return System; }
+	XrSession GetSession() override { return Session; }
+	XrTime GetDisplayTime() const override;
+	XrSpace GetTrackingSpace() const override;
+	IOpenXRExtensionPluginDelegates& GetIOpenXRExtensionPluginDelegates() override { return *this; }
+
 	OPENXRHMD_API TArray<IOpenXRExtensionPlugin*>& GetExtensionPlugins() { return ExtensionPlugins; }
 	OPENXRHMD_API void SetEnvironmentBlendMode(XrEnvironmentBlendMode NewBlendMode);
 
diff --git a/Engine/Plugins/Runtime/OpenXR/Source/OpenXRHMD/Public/IOpenXRHMD.h b/Engine/Plugins/Runtime/OpenXR/Source/OpenXRHMD/Public/IOpenXRHMD.h
new file mode 100644
--- /dev/null
+++ b/Engine/Plugins/Runtime/OpenXR/Source/OpenXRHMD/Public/IOpenXRHMD.h
@@ -0,0 +1,31 @@
+// Copyright Epic Games, Inc. All Rights Reserved.
+
+#pragma once
+
+#include "CoreMinimal.h"
+#include <openxr/openxr.h>
+
+class IOpenXRHMD
+{
+public:
+	OPENXRHMD_API virtual bool IsInitialized() const = 0;
+	OPENXRHMD_API virtual bool IsRunning() const = 0;
+	OPENXRHMD_API virtual bool IsFocused() const = 0;
+
+	OPENXRHMD_API virtual int32 AddTrackedDevice(XrAction Action, XrPath Path) = 0;
+	OPENXRHMD_API virtual void ResetTrackedDevices() = 0;
+	OPENXRHMD_API virtual XrPath GetTrackedDevicePath(const int32 DeviceId) = 0;
+	OPENXRHMD_API virtual XrSpace GetTrackedDeviceSpace(const int32 DeviceId) = 0;
+
+	OPENXRHMD_API virtual bool IsExtensionEnabled(const FString& Name) const = 0;
+	OPENXRHMD_API virtual XrInstance GetInstance() = 0;
+	OPENXRHMD_API virtual XrSystemId GetSystem() = 0;
+	OPENXRHMD_API virtual XrSession GetSession() = 0;
+	OPENXRHMD_API virtual XrTime GetDisplayTime() const = 0;
+	OPENXRHMD_API virtual XrSpace GetTrackingSpace() const = 0;
+
+	OPENXRHMD_API virtual class IOpenXRExtensionPluginDelegates& GetIOpenXRExtensionPluginDelegates() = 0;
+
+	OPENXRHMD_API virtual bool GetPoseForTime(int32 DeviceId, FTimespan Timespan, bool& OutTimeWasUsed, FQuat& CurrentOrientation, FVector& CurrentPosition, bool& bProvidedLinearVelocity, FVector& LinearVelocity, bool& bProvidedAngularVelocity, FVector& AngularVelocityAsAxisAndLength, bool& bProvidedLinearAcceleration, FVector& LinearAcceleration, float WorldToMetersScale) = 0;
+};
+
diff --git a/Engine/Plugins/Runtime/OpenXRViveTracker/Source/OpenXRViveTracker/OpenXRViveTracker.Build.cs b/Engine/Plugins/Runtime/OpenXRViveTracker/Source/OpenXRViveTracker/OpenXRViveTracker.Build.cs
--- a/Engine/Plugins/Runtime/OpenXRViveTracker/Source/OpenXRViveTracker/OpenXRViveTracker.Build.cs
+++ b/Engine/Plugins/Runtime/OpenXRViveTracker/Source/OpenXRViveTracker/OpenXRViveTracker.Build.cs
@@ -16,12 +16,6 @@ namespace UnrealBuildTool.Rules
 				}
 			 );
 
-			PrivateIncludePaths.AddRange(
-				new string[] {
-					Path.Combine(GetModuleDirectory("OpenXRHMD"), "Private"), // TODO: Adding private include path from other module
-				}
-			);
-
 			PrivateDependencyModuleNames.AddRange(
 				new string[]
 				{
diff --git a/Engine/Plugins/Runtime/OpenXRViveTracker/Source/OpenXRViveTracker/Private/OpenXRViveTracker.cpp b/Engine/Plugins/Runtime/OpenXRViveTracker/Source/OpenXRViveTracker/Private/OpenXRViveTracker.cpp
--- a/Engine/Plugins/Runtime/OpenXRViveTracker/Source/OpenXRViveTracker/Private/OpenXRViveTracker.cpp
+++ b/Engine/Plugins/Runtime/OpenXRViveTracker/Source/OpenXRViveTracker/Private/OpenXRViveTracker.cpp
@@ -16,6 +16,8 @@
 #include "Features/IModularFeatures.h"
 #include "HeadMountedDisplayFunctionLibrary.h"
 #include "IOpenXRHMDModule.h"
+#include "IOpenXRHMD.h"
+#include "IOpenXRExtensionPluginDelegates.h"
 
 #define LOCTEXT_NAMESPACE "OpenXRViveTracker"
 
@@ -111,7 +113,7 @@ FOpenXRViveTracker::FViveTracker::FViveTracker(XrActionSet InActionSet, FOpenXRP
 	XR_ENSURE(xrCreateAction(ActionSet, &Info, &VibrationAction));
 }
 
-void FOpenXRViveTracker::FViveTracker::AddTrackedDevices(FOpenXRHMD* HMD)
+void FOpenXRViveTracker::FViveTracker::AddTrackedDevices(IOpenXRHMD* HMD)
 {
 	if (HMD)
 	{
@@ -179,10 +181,10 @@ void FOpenXRViveTracker::PostCreateInstance(XrInstance InInstance)
 
 const void* FOpenXRViveTracker::OnCreateSession(XrInstance InInstance, XrSystemId InSystem, const void* InNext)
 {
-	static FName SystemName(TEXT("OpenXR"));
-	if (GEngine->XRSystem.IsValid() && (GEngine->XRSystem->GetSystemName() == SystemName))
+	if (GEngine && GEngine->XRSystem.IsValid())
 	{
-		OpenXRHMD = (FOpenXRHMD*)GEngine->XRSystem.Get();
+		XRTrackingSystem = GEngine->XRSystem.Get();
+		OpenXRHMD = XRTrackingSystem->GetIOpenXRHMD();
 	}
 
 	if (TrackerActionSet)
@@ -286,7 +288,7 @@ bool FOpenXRViveTracker::GetControllerOrientationAndPosition(const int32 Control
 	if (ControllerIndex == DeviceIndex && MotionSourceToEControllerHandMap.Contains(MotionSource))
 	{
 		FQuat Orientation;
-		bool Success = OpenXRHMD->GetCurrentPose(GetDeviceIDForMotionSource(MotionSource), Orientation, OutPosition);
+		bool Success = XRTrackingSystem->GetCurrentPose(GetDeviceIDForMotionSource(MotionSource), Orientation, OutPosition);
 		OutOrientation = FRotator(Orientation);
 		return Success;
 	}
@@ -464,7 +466,7 @@ void FOpenXRViveTracker::SetHapticFeedbackValues(int32 ControllerId, int32 Hand,
 			FOpenXRExtensionChainStructPtrs ScopedExtensionChainStructs;
 			if (Values.HapticBuffer != nullptr)
 			{
-				OpenXRHMD->GetApplyHapticFeedbackAddChainStructsDelegate().Broadcast(&HapticValue, ScopedExtensionChainStructs, Values.HapticBuffer);
+				OpenXRHMD->GetIOpenXRExtensionPluginDelegates().GetApplyHapticFeedbackAddChainStructsDelegate().Broadcast(&HapticValue, ScopedExtensionChainStructs, Values.HapticBuffer);
 			}
 			XR_ENSURE(xrApplyHapticFeedback(Session, &HapticActionInfo, (const XrHapticBaseHeader*)&HapticValue));
 		}
diff --git a/Engine/Plugins/Runtime/OpenXRViveTracker/Source/OpenXRViveTracker/Private/OpenXRViveTracker.h b/Engine/Plugins/Runtime/OpenXRViveTracker/Source/OpenXRViveTracker/Private/OpenXRViveTracker.h
--- a/Engine/Plugins/Runtime/OpenXRViveTracker/Source/OpenXRViveTracker/Private/OpenXRViveTracker.h
+++ b/Engine/Plugins/Runtime/OpenXRViveTracker/Source/OpenXRViveTracker/Private/OpenXRViveTracker.h
@@ -12,7 +12,6 @@
 
 #include "IOpenXRExtensionPlugin.h"
 #include "OpenXRCore.h"
-#include "OpenXRHMD.h"
 
 /**
   * OpenXR ViveTracker
@@ -34,7 +33,7 @@ public:
 
 		FViveTracker(XrActionSet InActionSet, FOpenXRPath InRolePath, const char* InName);
 
-		void AddTrackedDevices(class FOpenXRHMD* HMD);
+		void AddTrackedDevices(class IOpenXRHMD* HMD);
 		void GetSuggestedBindings(TArray<XrActionSuggestedBinding>& OutSuggestedBindings);
 	};
 
@@ -90,7 +89,8 @@ private:
 
 	PFN_xrEnumerateViveTrackerPathsHTCX xrEnumerateViveTrackerPathsHTCX = nullptr;
 
-	class FOpenXRHMD* OpenXRHMD = nullptr;
+	class IXRTrackingSystem* XRTrackingSystem = nullptr;
+	class IOpenXRHMD* OpenXRHMD = nullptr;
 
 	TSharedPtr<FGenericApplicationMessageHandler> MessageHandler;
 	int32 DeviceIndex;
diff --git a/Engine/Source/Runtime/HeadMountedDisplay/Public/IXRTrackingSystem.h b/Engine/Source/Runtime/HeadMountedDisplay/Public/IXRTrackingSystem.h
--- a/Engine/Source/Runtime/HeadMountedDisplay/Public/IXRTrackingSystem.h
+++ b/Engine/Source/Runtime/HeadMountedDisplay/Public/IXRTrackingSystem.h
@@ -15,6 +15,7 @@ class UARPin;
 class FSceneViewFamily;
 struct FWorldContext;
 class FARSupportInterface;
+class IOpenXRHMD;
 
 /**
  * Struct representing the properties of an external tracking sensor.
@@ -454,4 +455,9 @@ public:
 	 * Get the transform and dimensions of the area where the user can freely move while remaining tracked centered around the specified origin transform
 	 */
 	virtual bool GetPlayAreaRect(FTransform& OutTransform, FVector2D& OutRect) const { return false; }
+
+	/**
+	* Get the IOpenXRHMD interface, if there is one.
+	*/
+	virtual IOpenXRHMD* GetIOpenXRHMD() { return nullptr; }
 };
