UE-193722 OpenXRAR, OpenXRInput reference the Private folder of the OpenXRHMD module
-Refactored to use IXRTrackingSystem and IOpenXRHMD instead of FOpenXRHMD.

diff --git a/Engine/Plugins/Runtime/OpenXR/Source/OpenXRAR/OpenXRAR.Build.cs b/Engine/Plugins/Runtime/OpenXR/Source/OpenXRAR/OpenXRAR.Build.cs
--- a/Engine/Plugins/Runtime/OpenXR/Source/OpenXRAR/OpenXRAR.Build.cs
+++ b/Engine/Plugins/Runtime/OpenXR/Source/OpenXRAR/OpenXRAR.Build.cs
@@ -7,9 +7,6 @@ public class OpenXRAR : ModuleRules
 {
     public OpenXRAR(ReadOnlyTargetRules Target) : base(Target)
     {
-		var EngineDir = Path.GetFullPath(Target.RelativeEnginePath);
-		PrivateIncludePaths.Add(Path.Combine(GetModuleDirectory("OpenXRHMD"), "Private"));
-
 		PublicDependencyModuleNames.AddRange(
             new string[]
             {
diff --git a/Engine/Plugins/Runtime/OpenXR/Source/OpenXRAR/Private/OpenXRAR.cpp b/Engine/Plugins/Runtime/OpenXR/Source/OpenXRAR/Private/OpenXRAR.cpp
--- a/Engine/Plugins/Runtime/OpenXR/Source/OpenXRAR/Private/OpenXRAR.cpp
+++ b/Engine/Plugins/Runtime/OpenXR/Source/OpenXRAR/Private/OpenXRAR.cpp
@@ -1,7 +1,8 @@
 // Copyright Epic Games, Inc. All Rights Reserved.
 
 #include "OpenXRAR.h"
-#include "OpenXRHMD.h"
+#include "IXRTrackingSystem.h"
+#include "IOpenXRHMD.h"
 #include "IOpenXRExtensionPlugin.h"
 #include "MRMeshComponent.h"
 #include "ARLifeCycleComponent.h"
@@ -36,17 +37,14 @@ FOpenXRARSystem::~FOpenXRARSystem()
 	UARLifeCycleComponent::OnSpawnARActorDelegate.Remove(SpawnARActorDelegateHandle);
 }
 
-void FOpenXRARSystem::SetTrackingSystem(TSharedPtr<FXRTrackingSystemBase, ESPMode::ThreadSafe> InTrackingSystem)
+void FOpenXRARSystem::SetTrackingSystem(IXRTrackingSystem& InTrackingSystem)
 {
-	static FName SystemName(TEXT("OpenXR"));
-	if (InTrackingSystem->GetSystemName() == SystemName)
-	{
-		TrackingSystem = static_cast<FOpenXRHMD*>(InTrackingSystem.Get());
-	}
-
-	check(TrackingSystem != nullptr);
+	TrackingSystem = &InTrackingSystem;
+	check(TrackingSystem);
+	OpenXRHMD = TrackingSystem->GetIOpenXRHMD();
+	check(OpenXRHMD);
 
-	for (auto Plugin : TrackingSystem->GetExtensionPlugins())
+	for (auto Plugin : OpenXRHMD->GetExtensionPlugins())
 	{
 		if (CustomAnchorSupport == nullptr)
 		{
@@ -108,7 +106,7 @@ void FOpenXRARSystem::OnStartARSession(UARSessionConfig* InSessionConfig)
 
 	SessionStatus.Status = EARSessionStatus::Running;
 
-	for (auto Plugin : TrackingSystem->GetExtensionPlugins())
+	for (auto Plugin : OpenXRHMD->GetExtensionPlugins())
 	{
 		Plugin->OnStartARSession(InSessionConfig);
 	}
@@ -116,7 +114,7 @@ void FOpenXRARSystem::OnStartARSession(UARSessionConfig* InSessionConfig)
 
 void FOpenXRARSystem::OnStopARSession() 
 {
-	for (auto Plugin : TrackingSystem->GetExtensionPlugins())
+	for (auto Plugin : OpenXRHMD->GetExtensionPlugins())
 	{
 		Plugin->OnStopARSession();
 	}
@@ -131,7 +129,7 @@ void FOpenXRARSystem::OnStopARSession()
 
 void FOpenXRARSystem::OnPauseARSession() 
 {
-	for (auto Plugin : TrackingSystem->GetExtensionPlugins())
+	for (auto Plugin : OpenXRHMD->GetExtensionPlugins())
 	{
 		Plugin->OnPauseARSession();
 	}
@@ -245,15 +243,15 @@ UARPin* FOpenXRARSystem::OnPinComponent(USceneComponent* ComponentToPin, const F
 	{
 		if (CustomAnchorSupport != nullptr)
 		{
-			XrSession Session = TrackingSystem->GetSession();
+			XrSession Session = OpenXRHMD->GetSession();
 			XrTime DisplayTime = 0;
-			if (TrackingSystem->IsFocused())
+			if (OpenXRHMD->IsFocused())
 			{
 				// If the tracking system is not focused we will try to pin with time zero.
 				// The locate should fail, but we will still setup the pin so it will work when the tracking system is working.
-				DisplayTime = TrackingSystem->GetDisplayTime();
+				DisplayTime = OpenXRHMD->GetDisplayTime();
 			}
-			XrSpace TrackingSpace = TrackingSystem->GetTrackingSpace();
+			XrSpace TrackingSpace = OpenXRHMD->GetTrackingSpace();
 			float WorldToMetersScale = TrackingSystem->GetWorldToMetersScale();
 			if (!CustomAnchorSupport->OnPinComponent(NewPin, Session, TrackingSpace, DisplayTime, WorldToMetersScale))
 			{
@@ -292,11 +290,11 @@ void FOpenXRARSystem::UpdateAnchors()
 
 	if (CustomAnchorSupport != nullptr)
 	{
-		if (TrackingSystem->IsFocused())
+		if (OpenXRHMD->IsFocused())
 		{
-			XrSession Session = TrackingSystem->GetSession();
-			XrTime DisplayTime = TrackingSystem->GetDisplayTime();
-			XrSpace TrackingSpace = TrackingSystem->GetTrackingSpace();
+			XrSession Session = OpenXRHMD->GetSession();
+			XrTime DisplayTime = OpenXRHMD->GetDisplayTime();
+			XrSpace TrackingSpace = OpenXRHMD->GetTrackingSpace();
 			float WorldToMetersScale = TrackingSystem->GetWorldToMetersScale();
 			for (UARPin* Pin : Pins)
 			{
@@ -322,7 +320,7 @@ void FOpenXRARSystem::LoadARPins(TMap<FName, UARPin*>& LoadedPins)
 {
 	if (!IsLocalPinSaveSupported()) { return; }
 
-	CustomAnchorSupport->LoadARPins(TrackingSystem->GetSession(),
+	CustomAnchorSupport->LoadARPins(OpenXRHMD->GetSession(),
 		[&, this](FName Name)
 		{
 			check(IsInGameThread());
@@ -349,21 +347,21 @@ bool FOpenXRARSystem::SaveARPin(FName InName, UARPin* InPin)
 {
 	if (!IsLocalPinSaveSupported()) { return false; }
 
-	return CustomAnchorSupport->SaveARPin(TrackingSystem->GetSession(), InName, InPin);
+	return CustomAnchorSupport->SaveARPin(OpenXRHMD->GetSession(), InName, InPin);
 }
 
 void FOpenXRARSystem::RemoveSavedARPin(FName InName)
 {
 	if (!IsLocalPinSaveSupported()) { return; }
 
-	CustomAnchorSupport->RemoveSavedARPin(TrackingSystem->GetSession(), InName);
+	CustomAnchorSupport->RemoveSavedARPin(OpenXRHMD->GetSession(), InName);
 }
 
 void FOpenXRARSystem::RemoveAllSavedARPins()
 {
 	if (!IsLocalPinSaveSupported()) { return; }
 
-	CustomAnchorSupport->RemoveAllSavedARPins(TrackingSystem->GetSession());
+	CustomAnchorSupport->RemoveAllSavedARPins(OpenXRHMD->GetSession());
 }
 
 
@@ -1018,7 +1016,7 @@ IARSystemSupport* OpenXRARModuleImpl::CreateARSystem()
 	return ARSystemPtr;
 }
 
-void OpenXRARModuleImpl::SetTrackingSystem(TSharedPtr<FXRTrackingSystemBase, ESPMode::ThreadSafe> InTrackingSystem)
+void OpenXRARModuleImpl::SetTrackingSystem(IXRTrackingSystem& InTrackingSystem)
 {
 	ARSystem->SetTrackingSystem(InTrackingSystem);
 }
diff --git a/Engine/Plugins/Runtime/OpenXR/Source/OpenXRAR/Private/OpenXRAR.h b/Engine/Plugins/Runtime/OpenXR/Source/OpenXRAR/Private/OpenXRAR.h
--- a/Engine/Plugins/Runtime/OpenXR/Source/OpenXRAR/Private/OpenXRAR.h
+++ b/Engine/Plugins/Runtime/OpenXR/Source/OpenXRAR/Private/OpenXRAR.h
@@ -11,7 +11,8 @@
 #include "ARActor.h"
 #include "IOpenXRARModule.h"
 
-class FOpenXRHMD;
+class IOpenXRHMD;
+class IXRTrackingSystem;
 
 DECLARE_STATS_GROUP(TEXT("OpenXRAR"), STATGROUP_OPENXRAR, STATCAT_Advanced);
 
@@ -55,7 +56,7 @@ public:
 	FOpenXRARSystem();
 	virtual ~FOpenXRARSystem();
 
-	void SetTrackingSystem(TSharedPtr<FXRTrackingSystemBase, ESPMode::ThreadSafe> InTrackingSystem);
+	void SetTrackingSystem(IXRTrackingSystem& InTrackingSystem);
 
 	/** Invoked after the base AR system has been initialized. */
 	virtual void OnARSystemInitialized() {}
@@ -210,7 +211,8 @@ public:
 	}
 
 private:
-	FOpenXRHMD* TrackingSystem;
+	IXRTrackingSystem* TrackingSystem = nullptr;
+	IOpenXRHMD* OpenXRHMD = nullptr;
 
 	class IOpenXRCustomAnchorSupport* CustomAnchorSupport = nullptr;
 
@@ -301,7 +303,7 @@ class OpenXRARModuleImpl :
 public:
 	/** Used to init our AR system */
 	virtual IARSystemSupport* CreateARSystem();
-	virtual void SetTrackingSystem(TSharedPtr<FXRTrackingSystemBase, ESPMode::ThreadSafe> InTrackingSystem);
+	virtual void SetTrackingSystem(IXRTrackingSystem& InTrackingSystem);
 
 	virtual void StartupModule() override;
 
diff --git a/Engine/Plugins/Runtime/OpenXR/Source/OpenXRHMD/Private/OpenXRHMD.h b/Engine/Plugins/Runtime/OpenXR/Source/OpenXRHMD/Private/OpenXRHMD.h
--- a/Engine/Plugins/Runtime/OpenXR/Source/OpenXRHMD/Private/OpenXRHMD.h
+++ b/Engine/Plugins/Runtime/OpenXR/Source/OpenXRHMD/Private/OpenXRHMD.h
@@ -215,7 +215,7 @@ public:
 	virtual void ResetPosition() override;
 	virtual void Recenter(EOrientPositionSelector::Type Selector, float Yaw = 0.f);
 
-	virtual bool GetIsTracked(int32 DeviceId);
+	virtual bool GetIsTracked(int32 DeviceId) override;
 	virtual bool GetCurrentPose(int32 DeviceId, FQuat& CurrentOrientation, FVector& CurrentPosition) override;
 	virtual bool GetPoseForTime(int32 DeviceId, FTimespan Timespan, bool& OutTimeWasUsed, FQuat& CurrentOrientation, FVector& CurrentPosition, bool& bProvidedLinearVelocity, FVector& LinearVelocity, bool& bProvidedAngularVelocity, FVector& AngularVelocityAsAxisAndLength, bool& bProvidedLinearAcceleration, FVector& LinearAcceleration, float WorldToMetersScale) override;
 	virtual bool GetCurrentInteractionProfile(const EControllerHand Hand, FString& InteractionProfile) override;
@@ -397,10 +397,6 @@ public:
 	/** Constructor */
 	FOpenXRHMD(const FAutoRegister&, XrInstance InInstance, TRefCountPtr<FOpenXRRenderBridge>& InRenderBridge, TArray<const char*> InEnabledExtensions, TArray<class IOpenXRExtensionPlugin*> InExtensionPlugins, IARSystemSupport* ARSystemSupport);
 
-	void SetInputModule(IOpenXRInputModule* InInputModule)
-	{
-		InputModule = InInputModule;
-	}
 
 	/** Destructor */
 	virtual ~FOpenXRHMD();
@@ -410,6 +406,10 @@ public:
 	void OnFinishRendering_RHIThread();
 
 	/** IOpenXRHMD */
+	void SetInputModule(IOpenXRInputModule* InInputModule) override
+	{
+		InputModule = InInputModule;
+	}
 	/** @return	True if the HMD was initialized OK */
 	bool IsInitialized() const override;
 	bool IsRunning() const override;
@@ -425,8 +425,8 @@ public:
 	XrTime GetDisplayTime() const override;
 	XrSpace GetTrackingSpace() const override;
 	IOpenXRExtensionPluginDelegates& GetIOpenXRExtensionPluginDelegates() override { return *this; }
-
-	OPENXRHMD_API TArray<IOpenXRExtensionPlugin*>& GetExtensionPlugins() { return ExtensionPlugins; }
+	TArray<IOpenXRExtensionPlugin*>& GetExtensionPlugins() override { return ExtensionPlugins; }
+	
 	OPENXRHMD_API void SetEnvironmentBlendMode(XrEnvironmentBlendMode NewBlendMode);
 
 	/** Returns shader platform the plugin is currently configured for, in the editor it can change due to preview platforms. */
diff --git a/Engine/Plugins/Runtime/OpenXR/Source/OpenXRHMD/Private/OpenXRHMDModule.cpp b/Engine/Plugins/Runtime/OpenXR/Source/OpenXRHMD/Private/OpenXRHMDModule.cpp
--- a/Engine/Plugins/Runtime/OpenXR/Source/OpenXRHMD/Private/OpenXRHMDModule.cpp
+++ b/Engine/Plugins/Runtime/OpenXR/Source/OpenXRHMD/Private/OpenXRHMDModule.cpp
@@ -87,7 +87,7 @@ TSharedPtr< class IXRTrackingSystem, ESPMode::ThreadSafe > FOpenXRHMDModule::Cre
 	auto OpenXRHMD = FSceneViewExtensions::NewExtension<FOpenXRHMD>(Instance, RenderBridge, EnabledExtensions, ExtensionPlugins, ARSystem);
 	if (OpenXRHMD->IsInitialized())
 	{
-		ARModule->SetTrackingSystem(OpenXRHMD);
+		ARModule->SetTrackingSystem(OpenXRHMD.Get());
 		OpenXRHMD->GetARCompositionComponent()->InitializeARSystem();
 		return OpenXRHMD;
 	}
diff --git a/Engine/Plugins/Runtime/OpenXR/Source/OpenXRHMD/Public/IOpenXRARModule.h b/Engine/Plugins/Runtime/OpenXR/Source/OpenXRHMD/Public/IOpenXRARModule.h
--- a/Engine/Plugins/Runtime/OpenXR/Source/OpenXRHMD/Public/IOpenXRARModule.h
+++ b/Engine/Plugins/Runtime/OpenXR/Source/OpenXRHMD/Public/IOpenXRARModule.h
@@ -12,7 +12,7 @@ class IOpenXRARModule : public IModuleInterface, public IModularFeature
 public:
 	/** Used to init our AR system */
 	virtual class IARSystemSupport* CreateARSystem() = 0;
-	virtual void SetTrackingSystem(TSharedPtr<class FXRTrackingSystemBase, ESPMode::ThreadSafe> InTrackingSystem) = 0;
+	virtual void SetTrackingSystem(class IXRTrackingSystem& InTrackingSystem) = 0;
 	virtual bool GetExtensions(TArray<const ANSICHAR*>& OutExtensions) = 0;
 
 	virtual class IOpenXRARTrackedMeshHolder * GetTrackedMeshHolder() = 0;
diff --git a/Engine/Plugins/Runtime/OpenXR/Source/OpenXRHMD/Public/IOpenXRHMD.h b/Engine/Plugins/Runtime/OpenXR/Source/OpenXRHMD/Public/IOpenXRHMD.h
--- a/Engine/Plugins/Runtime/OpenXR/Source/OpenXRHMD/Public/IOpenXRHMD.h
+++ b/Engine/Plugins/Runtime/OpenXR/Source/OpenXRHMD/Public/IOpenXRHMD.h
@@ -2,12 +2,18 @@
 
 #pragma once
 
+// This interface allows OpenXR specific modules to access an OpenXR specific api for OpenXRHMD, avoiding the need
+// to directly access FOpenXRHMD.  
+// See also IXRTrackingSystem::GetIOpenXRHMD().
+
 #include "CoreMinimal.h"
 #include <openxr/openxr.h>
 
 class IOpenXRHMD
 {
 public:
+	OPENXRHMD_API virtual void SetInputModule(class IOpenXRInputModule* InInputModule) = 0;
+		
 	OPENXRHMD_API virtual bool IsInitialized() const = 0;
 	OPENXRHMD_API virtual bool IsRunning() const = 0;
 	OPENXRHMD_API virtual bool IsFocused() const = 0;
@@ -26,6 +32,9 @@ public:
 
 	OPENXRHMD_API virtual class IOpenXRExtensionPluginDelegates& GetIOpenXRExtensionPluginDelegates() = 0;
 
+	OPENXRHMD_API virtual bool GetIsTracked(int32 DeviceId) = 0;
 	OPENXRHMD_API virtual bool GetPoseForTime(int32 DeviceId, FTimespan Timespan, bool& OutTimeWasUsed, FQuat& CurrentOrientation, FVector& CurrentPosition, bool& bProvidedLinearVelocity, FVector& LinearVelocity, bool& bProvidedAngularVelocity, FVector& AngularVelocityAsAxisAndLength, bool& bProvidedLinearAcceleration, FVector& LinearAcceleration, float WorldToMetersScale) = 0;
+
+	OPENXRHMD_API virtual TArray<class IOpenXRExtensionPlugin*>& GetExtensionPlugins() = 0;
 };
 
diff --git a/Engine/Plugins/Runtime/OpenXR/Source/OpenXRInput/OpenXRInput.Build.cs b/Engine/Plugins/Runtime/OpenXR/Source/OpenXRInput/OpenXRInput.Build.cs
--- a/Engine/Plugins/Runtime/OpenXR/Source/OpenXRInput/OpenXRInput.Build.cs
+++ b/Engine/Plugins/Runtime/OpenXR/Source/OpenXRInput/OpenXRInput.Build.cs
@@ -8,8 +8,6 @@ namespace UnrealBuildTool.Rules
     {
         public OpenXRInput(ReadOnlyTargetRules Target) : base(Target)
         {
-            PrivateIncludePaths.Add(Path.Combine(GetModuleDirectory("OpenXRHMD"), "Private"));
-
             PrivateDependencyModuleNames.AddRange(
                 new string[]
                 {
diff --git a/Engine/Plugins/Runtime/OpenXR/Source/OpenXRInput/Private/OpenXRInput.cpp b/Engine/Plugins/Runtime/OpenXR/Source/OpenXRInput/Private/OpenXRInput.cpp
--- a/Engine/Plugins/Runtime/OpenXR/Source/OpenXRInput/Private/OpenXRInput.cpp
+++ b/Engine/Plugins/Runtime/OpenXR/Source/OpenXRInput/Private/OpenXRInput.cpp
@@ -2,11 +2,13 @@
 
 #include "OpenXRInput.h"
 #include "OpenXRInputSettings.h"
-#include "OpenXRHMD.h"
+#include "IOpenXRHMD.h"
+#include "IXRTrackingSystem.h"
 #include "OpenXRCore.h"
 #include "UObject/UObjectIterator.h"
 #include "GameFramework/InputSettings.h"
 #include "IOpenXRExtensionPlugin.h"
+#include "IOpenXRExtensionPluginDelegates.h"
 #include "HeadMountedDisplayFunctionLibrary.h"
 
 #include "EnhancedInputLibrary.h"
@@ -74,29 +76,17 @@ FOpenXRInputPlugin::~FOpenXRInputPlugin()
 {
 }
 
-FOpenXRHMD* FOpenXRInputPlugin::GetOpenXRHMD() const
-{
-	static FName SystemName(TEXT("OpenXR"));
-	if (GEngine->XRSystem.IsValid() && (GEngine->XRSystem->GetSystemName() == SystemName))
-	{
-		return static_cast<FOpenXRHMD*>(GEngine->XRSystem.Get());
-	}
-
-	return nullptr;
-}
-
 void FOpenXRInputPlugin::StartupModule()
 {
 	IOpenXRInputPlugin::StartupModule();
 
-	FOpenXRHMD* OpenXRHMD = GetOpenXRHMD();
-	// Note: OpenXRHMD may be null, for example in the editor.  But we still need the input device to enumerate sources.
-	InputDevice = MakeShared<FOpenXRInput>(OpenXRHMD);
+	// Note: XRSystem may be null, for example in the editor.  But we still need the input device to enumerate sources.
+	InputDevice = MakeShared<FOpenXRInput>(GEngine->XRSystem.Get());
 }
 
 FOpenXRInputPlugin::FOpenXRAction::FOpenXRAction(XrActionSet InActionSet,
 	XrActionType InActionType, const FName& InName, const FString& InLocalizedName,
-	const TArray<XrPath>& InSubactionPaths, const TObjectPtr<const UInputAction>& InObject, FOpenXRHMD* OpenXRHMD)
+	const TArray<XrPath>& InSubactionPaths, const TObjectPtr<const UInputAction>& InObject, IOpenXRHMD* OpenXRHMD)
 	: FOpenXRAction(InActionSet, InActionType, InName, InLocalizedName, InSubactionPaths, OpenXRHMD)
 {
 	Object = InObject;
@@ -104,7 +94,7 @@ FOpenXRInputPlugin::FOpenXRAction::FOpenXRAction(XrActionSet InActionSet,
 
 FOpenXRInputPlugin::FOpenXRAction::FOpenXRAction(XrActionSet InActionSet,
 	XrActionType InActionType, const FName& InName, const FString& InLocalizedName,
-	const TArray<XrPath>& InSubactionPaths, FOpenXRHMD* OpenXRHMD)
+	const TArray<XrPath>& InSubactionPaths, IOpenXRHMD* OpenXRHMD)
 	: Set(InActionSet)
 	, Type(InActionType)
 	, Name(InName)
@@ -151,14 +141,14 @@ FOpenXRInputPlugin::FOpenXRAction::FOpenXRAction(XrActionSet InActionSet,
 
 FOpenXRInputPlugin::FOpenXRActionSet::FOpenXRActionSet(XrInstance InInstance,
 	const FName& InName, const FString& InLocalizedName, uint32 InPriority,
-	const TObjectPtr<const UInputMappingContext>& InObject, FOpenXRHMD* OpenXRHMD)
+	const TObjectPtr<const UInputMappingContext>& InObject, IOpenXRHMD* OpenXRHMD)
 	: FOpenXRActionSet(InInstance, InName, InLocalizedName, InPriority, OpenXRHMD)
 {
 	Object = InObject;
 }
 
 FOpenXRInputPlugin::FOpenXRActionSet::FOpenXRActionSet(XrInstance InInstance,
-	const FName& InName, const FString& InLocalizedName, uint32 InPriority, FOpenXRHMD* OpenXRHMD)
+	const FName& InName, const FString& InLocalizedName, uint32 InPriority, IOpenXRHMD* OpenXRHMD)
 	: Handle(XR_NULL_HANDLE)
 	, Name(InName)
 	, LocalizedName(InLocalizedName)
@@ -243,7 +233,7 @@ FOpenXRInputPlugin::FOpenXRController::FOpenXRController(XrActionSet InActionSet
 	XR_ENSURE(xrCreateAction(ActionSet, &Info, &VibrationAction));
 }
 
-void FOpenXRInputPlugin::FOpenXRController::AddTrackedDevices(FOpenXRHMD* HMD)
+void FOpenXRInputPlugin::FOpenXRController::AddTrackedDevices(IOpenXRHMD* HMD)
 {
 	if (HMD)
 	{
@@ -260,8 +250,9 @@ FOpenXRInputPlugin::FInteractionProfile::FInteractionProfile(XrPath InProfile, b
 {
 }
 
-FOpenXRInputPlugin::FOpenXRInput::FOpenXRInput(FOpenXRHMD* HMD)
-	: OpenXRHMD(HMD)
+FOpenXRInputPlugin::FOpenXRInput::FOpenXRInput(IXRTrackingSystem* InTrackingSystem)
+	: TrackingSystem(InTrackingSystem)
+	, OpenXRHMD(InTrackingSystem ? InTrackingSystem->GetIOpenXRHMD() : nullptr)
 	, Instance(XR_NULL_HANDLE)
 	, ControllerActionSet()
 	, ActionSets()
@@ -278,7 +269,7 @@ FOpenXRInputPlugin::FOpenXRInput::FOpenXRInput(FOpenXRHMD* HMD)
 	, MessageHandler(new FGenericApplicationMessageHandler())
 {
 	IModularFeatures::Get().RegisterModularFeature(GetModularFeatureName(), this);
-	
+
 	// If there is no HMD then this module is not active, but it still needs to exist so we can EnumerateMotionSources from it.
 	if (OpenXRHMD)
 	{
@@ -1221,7 +1212,7 @@ bool FOpenXRInputPlugin::FOpenXRInput::GetControllerOrientationAndPosition(const
 		if (Result >= XR_SUCCESS && State.isActive)
 		{
 			FQuat Orientation;
-			OpenXRHMD->GetCurrentPose(GetDeviceIDForMotionSource(MotionSource), Orientation, OutPosition);
+			TrackingSystem->GetCurrentPose(GetDeviceIDForMotionSource(MotionSource), Orientation, OutPosition);
 			OutOrientation = FRotator(Orientation);
 			return true;
 		}
@@ -1437,7 +1428,7 @@ void FOpenXRInputPlugin::FOpenXRInput::SetHapticFeedbackValues(int32 ControllerI
 				FOpenXRExtensionChainStructPtrs ScopedExtensionChainStructs;
 				if (Values.HapticBuffer != nullptr)
 				{
-					OpenXRHMD->GetApplyHapticFeedbackAddChainStructsDelegate().Broadcast(&HapticValue, ScopedExtensionChainStructs, Values.HapticBuffer);
+					OpenXRHMD->GetIOpenXRExtensionPluginDelegates().GetApplyHapticFeedbackAddChainStructsDelegate().Broadcast(&HapticValue, ScopedExtensionChainStructs, Values.HapticBuffer);
 				}
 				XR_ENSURE(xrApplyHapticFeedback(Session, &HapticActionInfo, (const XrHapticBaseHeader*)&HapticValue));
 
diff --git a/Engine/Plugins/Runtime/OpenXR/Source/OpenXRInput/Private/OpenXRInput.h b/Engine/Plugins/Runtime/OpenXR/Source/OpenXRInput/Private/OpenXRInput.h
--- a/Engine/Plugins/Runtime/OpenXR/Source/OpenXRInput/Private/OpenXRInput.h
+++ b/Engine/Plugins/Runtime/OpenXR/Source/OpenXRInput/Private/OpenXRInput.h
@@ -11,7 +11,8 @@
 #include "UObject/ObjectPtr.h"
 #include "UObject/StrongObjectPtr.h"
 
-class FOpenXRHMD;
+class IOpenXRHMD;
+class IXRTrackingSystem;
 class UInputAction;
 class UInputTrigger;
 class UInputModifier;
@@ -44,14 +45,14 @@ public:
 			FString& InLocalizedName,
 			const TArray<XrPath>& InSubactionPaths,
 			const TObjectPtr<const UInputAction>& InObject,
-			FOpenXRHMD* OpenXRHMD = nullptr);
+			IOpenXRHMD* OpenXRHMD = nullptr);
 
 		FOpenXRAction(XrActionSet InActionSet,
 			XrActionType InActionType,
 			const FName& InName,
 			const FString& InLocalizedName,
 			const TArray<XrPath>& InSubactionPaths,
-			FOpenXRHMD* OpenXRHMD = nullptr);
+			IOpenXRHMD* OpenXRHMD = nullptr);
 	};
 
 	struct FOpenXRActionSet
@@ -67,13 +68,13 @@ public:
 			const FString& InLocalizedName,
 			uint32 InPriority,
 			const TObjectPtr<const UInputMappingContext>& InObject,
-			FOpenXRHMD* OpenXRHMD = nullptr);
+			IOpenXRHMD* OpenXRHMD = nullptr);
 
 		FOpenXRActionSet(XrInstance InInstance,
 			const FName& InName,
 			const FString& InLocalizedName,
 			uint32 InPriority,
-			FOpenXRHMD* OpenXRHMD = nullptr);
+			IOpenXRHMD* OpenXRHMD = nullptr);
 	};
 
 	struct FOpenXRController
@@ -92,7 +93,7 @@ public:
 
 		FOpenXRController(XrActionSet InActionSet, XrPath InUserPath, const char* InName);
 
-		void AddTrackedDevices(FOpenXRHMD* HMD);
+		void AddTrackedDevices(IOpenXRHMD* HMD);
 	};
 
 	struct FInteractionProfile
@@ -108,7 +109,7 @@ public:
 	class FOpenXRInput : public IOpenXRInputModule, public IInputDevice, public FXRMotionControllerBase, public IHapticDevice, public TSharedFromThis<FOpenXRInput>
 	{
 	public:
-		FOpenXRInput(FOpenXRHMD* HMD);
+		FOpenXRInput(IXRTrackingSystem* InTrackingSystem);
 		virtual ~FOpenXRInput() {};
 
 		// IOpenXRAdditionalModule overrides
@@ -144,7 +145,8 @@ public:
 		virtual float GetHapticAmplitudeScale() const override;
 
 	private:
-		FOpenXRHMD* OpenXRHMD;
+		IXRTrackingSystem* TrackingSystem;
+		IOpenXRHMD* OpenXRHMD;
 		XrInstance Instance;
 
 		TUniquePtr<FOpenXRActionSet> ControllerActionSet;
@@ -196,9 +198,6 @@ public:
 	virtual void StartupModule() override;
 	virtual TSharedPtr< class IInputDevice > CreateInputDevice(const TSharedRef< FGenericApplicationMessageHandler >& InMessageHandler) override;
 
-private:
-	FOpenXRHMD* GetOpenXRHMD() const;
-
 private:
 	TSharedPtr<FOpenXRInput> InputDevice;
 };
