UE-200740 OpenXR XrSystem loss/change handling.
-Removed unnecessary xrGetSystem call in the FOpenXRHMD constructor.  It is overwritten (for example in OnStereoStartup()) without being used.

diff --git a/Engine/Plugins/Runtime/OpenXR/Source/OpenXRHMD/Private/OpenXRHMD.cpp b/Engine/Plugins/Runtime/OpenXR/Source/OpenXRHMD/Private/OpenXRHMD.cpp
--- a/Engine/Plugins/Runtime/OpenXR/Source/OpenXRHMD/Private/OpenXRHMD.cpp
+++ b/Engine/Plugins/Runtime/OpenXR/Source/OpenXRHMD/Private/OpenXRHMD.cpp
@@ -859,6 +859,8 @@ bool FOpenXRHMD::EnableStereo(bool stereo)
 			{
 				GEngine->SetMaxFPS(0);
 			}
+
+			// Note: This StartSession may not work, but if not we should receive a SESSION_STATE_READY and try again or a LOSS_PENDING and session destruction
 			StartSession();
 
 			FApp::SetUseVRFocus(true);
@@ -1310,7 +1312,7 @@ FOpenXRHMD::FOpenXRHMD(const FAutoRegister& AutoRegister, XrInstance InInstance,
 	, InputModule(nullptr)
 	, ExtensionPlugins(std::move(InExtensionPlugins))
 	, Instance(InInstance)
-	, System(IOpenXRHMDModule::Get().GetSystemId())
+	, System(XR_NULL_SYSTEM_ID)
 	, Session(XR_NULL_HANDLE)
 	, LocalSpace(XR_NULL_HANDLE)
 	, StageSpace(XR_NULL_HANDLE)
