UE-205956 OpenXR: Need ability to add subaction paths to action.
-Added SubactionPath parameter to IOpenXRHMD::AddTrackedDevice and FOpenXRHMD::FDevice so we can create actions with a SubactionPath specifed.

diff --git a/Engine/Plugins/Runtime/OpenXR/Source/OpenXRHMD/Private/OpenXRHMD.cpp b/Engine/Plugins/Runtime/OpenXR/Source/OpenXRHMD/Private/OpenXRHMD.cpp
--- a/Engine/Plugins/Runtime/OpenXR/Source/OpenXRHMD/Private/OpenXRHMD.cpp
+++ b/Engine/Plugins/Runtime/OpenXR/Source/OpenXRHMD/Private/OpenXRHMD.cpp
@@ -2172,15 +2172,18 @@ void FOpenXRHMD::DestroySession()
 		bNeedReBuildOcclusionMesh = true;
 	}
 }
-
 int32 FOpenXRHMD::AddTrackedDevice(XrAction Action, XrPath Path)
+{
+	return AddTrackedDevice(Action, Path, XR_NULL_PATH);
+}
+int32 FOpenXRHMD::AddTrackedDevice(XrAction Action, XrPath Path, XrPath SubactionPath)
 {
 	FWriteScopeLock DeviceLock(DeviceMutex);
 
 	// Ensure the HMD device is already emplaced
 	ensure(DeviceSpaces.Num() > 0);
 
-	int32 DeviceId = DeviceSpaces.Emplace(Action, Path);
+	int32 DeviceId = DeviceSpaces.Emplace(Action, Path, SubactionPath);
 
 	//FReadScopeLock Lock(SessionHandleMutex); // This is called from StartSession(), which already has this lock.
 	if (Session)
@@ -4169,6 +4172,15 @@ FOpenXRHMD::FDeviceSpace::FDeviceSpace(XrAction InAction, XrPath InPath)
 	: Action(InAction)
 	, Space(XR_NULL_HANDLE)
 	, Path(InPath)
+	, SubactionPath(XR_NULL_PATH)
+{
+}
+
+FOpenXRHMD::FDeviceSpace::FDeviceSpace(XrAction InAction, XrPath InPath, XrPath InSubactionPath)
+	: Action(InAction)
+	, Space(XR_NULL_HANDLE)
+	, Path(InPath)
+	, SubactionPath(InSubactionPath)
 {
 }
 
@@ -4187,7 +4199,7 @@ bool FOpenXRHMD::FDeviceSpace::CreateSpace(XrSession InSession)
 	XrActionSpaceCreateInfo ActionSpaceInfo;
 	ActionSpaceInfo.type = XR_TYPE_ACTION_SPACE_CREATE_INFO;
 	ActionSpaceInfo.next = nullptr;
-	ActionSpaceInfo.subactionPath = XR_NULL_PATH;
+	ActionSpaceInfo.subactionPath = SubactionPath;
 	ActionSpaceInfo.poseInActionSpace = ToXrPose(FTransform::Identity);
 	ActionSpaceInfo.action = Action;
 	return XR_ENSURE(xrCreateActionSpace(InSession, &ActionSpaceInfo, &Space));
diff --git a/Engine/Plugins/Runtime/OpenXR/Source/OpenXRHMD/Private/OpenXRHMD.h b/Engine/Plugins/Runtime/OpenXR/Source/OpenXRHMD/Private/OpenXRHMD.h
--- a/Engine/Plugins/Runtime/OpenXR/Source/OpenXRHMD/Private/OpenXRHMD.h
+++ b/Engine/Plugins/Runtime/OpenXR/Source/OpenXRHMD/Private/OpenXRHMD.h
@@ -50,6 +50,7 @@ public:
 	{
 	public:
 		FDeviceSpace(XrAction InAction, XrPath InPath);
+		FDeviceSpace(XrAction InAction, XrPath InPath, XrPath InSubactionPath);
 		~FDeviceSpace();
 
 		bool CreateSpace(XrSession InSession);
@@ -58,6 +59,7 @@ public:
 		XrAction Action;
 		XrSpace Space;
 		XrPath Path;
+		XrPath SubactionPath;
 	};
 
 	class FTrackingSpace
@@ -405,6 +407,7 @@ public:
 	bool IsRunning() const override;
 	bool IsFocused() const override;
 	int32 AddTrackedDevice(XrAction Action, XrPath Path) override;
+	int32 AddTrackedDevice(XrAction Action, XrPath Path, XrPath SubActionPath) override;
 	void ResetTrackedDevices() override;
 	XrPath GetTrackedDevicePath(const int32 DeviceId) override;
 	XrSpace GetTrackedDeviceSpace(const int32 DeviceId) override;
diff --git a/Engine/Plugins/Runtime/OpenXR/Source/OpenXRHMD/Public/IOpenXRHMD.h b/Engine/Plugins/Runtime/OpenXR/Source/OpenXRHMD/Public/IOpenXRHMD.h
--- a/Engine/Plugins/Runtime/OpenXR/Source/OpenXRHMD/Public/IOpenXRHMD.h
+++ b/Engine/Plugins/Runtime/OpenXR/Source/OpenXRHMD/Public/IOpenXRHMD.h
@@ -19,6 +19,7 @@ public:
 	OPENXRHMD_API virtual bool IsFocused() const = 0;
 
 	OPENXRHMD_API virtual int32 AddTrackedDevice(XrAction Action, XrPath Path) = 0;
+	OPENXRHMD_API virtual int32 AddTrackedDevice(XrAction Action, XrPath Path, XrPath SubactionPath) = 0;
 	OPENXRHMD_API virtual void ResetTrackedDevices() = 0;
 	OPENXRHMD_API virtual XrPath GetTrackedDevicePath(const int32 DeviceId) = 0;
 	OPENXRHMD_API virtual XrSpace GetTrackedDeviceSpace(const int32 DeviceId) = 0;
