Remove oculus-specific workaround.

diff --git a/Engine/Plugins/Runtime/OpenXR/Source/OpenXRHMD/Private/OpenXRHMD.cpp b/Engine/Plugins/Runtime/OpenXR/Source/OpenXRHMD/Private/OpenXRHMD.cpp
--- a/Engine/Plugins/Runtime/OpenXR/Source/OpenXRHMD/Private/OpenXRHMD.cpp
+++ b/Engine/Plugins/Runtime/OpenXR/Source/OpenXRHMD/Private/OpenXRHMD.cpp
@@ -3562,11 +3562,6 @@ void FOpenXRHMD::OnFinishRendering_RHIThread()
 			{
 				Layer.next = Module->OnEndProjectionLayer(Session, 0, Layer.next, Layer.layerFlags);
 			}
-
-#if PLATFORM_ANDROID
-			// @todo: temporary workaround for Quest compositor issue, see UE-145546
-			Layer.layerFlags |= XR_COMPOSITION_LAYER_UNPREMULTIPLIED_ALPHA_BIT;
-#endif
 		}
 		
 		XrCompositionLayerProjection CompositedLayer = {};
