Swap OpenXR input to use Project Settings->Engine->Enhanced Input for its Input Mapping Contexts rather than Plugins -> OpenXR Input.
Hook up the priority settings on the EI mapping contexts so they're passed into xrCreateActionSet.
The OpenXR Input Settings page is removed, since it's now redundant.

diff --git a/Engine/Plugins/Runtime/OpenXR/Source/OpenXREditor/Private/OpenXREditorModule.cpp b/Engine/Plugins/Runtime/OpenXR/Source/OpenXREditor/Private/OpenXREditorModule.cpp
--- a/Engine/Plugins/Runtime/OpenXR/Source/OpenXREditor/Private/OpenXREditorModule.cpp
+++ b/Engine/Plugins/Runtime/OpenXR/Source/OpenXREditor/Private/OpenXREditorModule.cpp
@@ -2,7 +2,6 @@
 
 #include "OpenXREditorModule.h"
 #include "OpenXRAssetDirectory.h"
-#include "OpenXRInputSettings.h"
 #include "OpenXRHMDSettings.h"
 
 #include "ISettingsModule.h"
@@ -21,12 +20,6 @@ void FOpenXREditorModule::StartupModule()
 
 	if (SettingsModule != nullptr)
 	{
-		SettingsModule->RegisterSettings("Project", "Plugins", "OpenXR",
-			LOCTEXT("OpenXRInputSettingsName", "OpenXR Input"),
-			LOCTEXT("OpenXRInputSettingsDescription", "Project settings for OpenXR plugin"),
-			GetMutableDefault<UOpenXRInputSettings>()
-		);
-
 		SettingsModule->RegisterSettings("Project", "Plugins", "OpenXRHMD",
 			LOCTEXT("OpenXRHMDSettingsName", "OpenXR Settings"),
 			LOCTEXT("OpenXRHMDSettingsDescription", "Project settings for OpenXR plugin"),
diff --git a/Engine/Plugins/Runtime/OpenXR/Source/OpenXRInput/Private/OpenXRInput.cpp b/Engine/Plugins/Runtime/OpenXR/Source/OpenXRInput/Private/OpenXRInput.cpp
--- a/Engine/Plugins/Runtime/OpenXR/Source/OpenXRInput/Private/OpenXRInput.cpp
+++ b/Engine/Plugins/Runtime/OpenXR/Source/OpenXRInput/Private/OpenXRInput.cpp
@@ -1,7 +1,7 @@
 // Copyright Epic Games, Inc. All Rights Reserved.
 
 #include "OpenXRInput.h"
-#include "OpenXRInputSettings.h"
+#include "EnhancedInputDeveloperSettings.h"
 #include "IOpenXRHMD.h"
 #include "IXRTrackingSystem.h"
 #include "OpenXRCore.h"
@@ -263,7 +263,7 @@ FOpenXRInputPlugin::FOpenXRInput::FOpenXRInput(IXRTrackingSystem* InTrackingSyst
 	, EnhancedActions()
 	, Controllers()
 	, MotionSourceToControllerHandMap()
-	, InputMappingContexts()
+	, InputMappingContextToPriorityMap()
 	, bActionsAttached(false)
 	, bDirectionalBindingSupported(false)
 	, bPalmPoseSupported(false)
@@ -403,17 +403,17 @@ bool FOpenXRInputPlugin::FOpenXRInput::BuildActions(XrSession Session)
 
 	PRAGMA_DISABLE_DEPRECATION_WARNINGS
 	// Attempt to load the default input config from the OpenXR input settings.
-	UOpenXRInputSettings* InputSettings = GetMutableDefault<UOpenXRInputSettings>();
+	const UEnhancedInputDeveloperSettings* InputSettings = GetDefault<UEnhancedInputDeveloperSettings>();
 	if (InputSettings)
 	{
-		for (const auto& Context : InputSettings->InputMappingContexts)
+		for (const auto& Context : InputSettings->DefaultMappingContexts)
 		{
-			TStrongObjectPtr<UInputMappingContext> Obj(Context.LoadSynchronous());
-			InputMappingContexts.Add(Obj);
+			TStrongObjectPtr<const UInputMappingContext> Obj(Context.InputMappingContext.LoadSynchronous());
+			InputMappingContextToPriorityMap.Add(Obj, Context.Priority);
 		}
 	}
 
-	if (!InputMappingContexts.IsEmpty())
+	if (!InputMappingContextToPriorityMap.IsEmpty())
 	{
 		BuildEnhancedActions(Profiles);
 	}
@@ -607,12 +607,12 @@ void FOpenXRInputPlugin::FOpenXRInput::BuildEnhancedActions(TMap<FString, FInter
 	}
 
 	PRAGMA_DISABLE_DEPRECATION_WARNINGS
-	for (const auto& MappingContext : InputMappingContexts)
+	for (const auto& MappingContext : InputMappingContextToPriorityMap)
 	{
-		FOpenXRActionSet ActionSet(Instance, MappingContext->GetFName(), MappingContext->ContextDescription.ToString(), 0, MappingContext.Get());
+		FOpenXRActionSet ActionSet(Instance, MappingContext.Key->GetFName(), MappingContext.Key->ContextDescription.ToString(), MappingContext.Value, MappingContext.Key.Get());
 		TMap<FName, int32> ActionMap;
 
-		for (const FEnhancedActionKeyMapping& Mapping : MappingContext->GetMappings())
+		for (const FEnhancedActionKeyMapping& Mapping : MappingContext.Key->GetMappings())
 		{
 			if (!Mapping.Action)
 			{
@@ -796,7 +796,7 @@ void FOpenXRInputPlugin::FOpenXRInput::OnDestroySession()
 	{
 		// If the session shut down, clean up.
 		bActionsAttached = false;
-		InputMappingContexts.Reset();
+		InputMappingContextToPriorityMap.Reset();
 	}
 }
 
@@ -1482,7 +1482,7 @@ bool FOpenXRInputPlugin::FOpenXRInput::AttachInputMappingContexts(const TSet<TOb
 
 	for (const auto& Context : MappingContexts)
 	{
-		InputMappingContexts.Add(TStrongObjectPtr<UInputMappingContext>(Context));
+		InputMappingContextToPriorityMap.Add(TStrongObjectPtr<UInputMappingContext>(Context), 0);
 	}
 	return true;
 }
diff --git a/Engine/Plugins/Runtime/OpenXR/Source/OpenXRInput/Private/OpenXRInput.h b/Engine/Plugins/Runtime/OpenXR/Source/OpenXRInput/Private/OpenXRInput.h
--- a/Engine/Plugins/Runtime/OpenXR/Source/OpenXRInput/Private/OpenXRInput.h
+++ b/Engine/Plugins/Runtime/OpenXR/Source/OpenXRInput/Private/OpenXRInput.h
@@ -156,7 +156,9 @@ public:
 		TMap<EControllerHand, FOpenXRController> Controllers;
 		TMap<FName, EControllerHand> MotionSourceToControllerHandMap;
 
-		TSet<TStrongObjectPtr<UInputMappingContext>> InputMappingContexts;
+		// Holds Enhanced Input Mapping Contexts pulled from the Enhanced Input Settings and mapped to their priority
+		// These will be converted into OpenXR action sets
+		TMap<TStrongObjectPtr<const UInputMappingContext>, uint32> InputMappingContextToPriorityMap;
 
 		XrAction GetActionForMotionSource(FName MotionSource) const;
 		int32 GetDeviceIDForMotionSource(FName MotionSource) const;
diff --git a/Engine/Plugins/Runtime/OpenXR/Source/OpenXRInput/Private/OpenXRInputFunctionLibrary.cpp b/Engine/Plugins/Runtime/OpenXR/Source/OpenXRInput/Private/OpenXRInputFunctionLibrary.cpp
--- a/Engine/Plugins/Runtime/OpenXR/Source/OpenXRInput/Private/OpenXRInputFunctionLibrary.cpp
+++ b/Engine/Plugins/Runtime/OpenXR/Source/OpenXRInput/Private/OpenXRInputFunctionLibrary.cpp
@@ -1,7 +1,6 @@
 // Copyright Epic Games, Inc. All Rights Reserved.
 
 #include "OpenXRInputFunctionLibrary.h"
-#include "OpenXRInputSettings.h"
 
 #include "Engine/Engine.h"
 #include "Features/IModularFeatures.h"
diff --git a/Engine/Plugins/Runtime/OpenXR/Source/OpenXRInput/Private/OpenXRInputSettings.cpp b/Engine/Plugins/Runtime/OpenXR/Source/OpenXRInput/Private/OpenXRInputSettings.cpp
deleted file mode 100644
diff --git a/Engine/Plugins/Runtime/OpenXR/Source/OpenXRInput/Public/OpenXRInputSettings.h b/Engine/Plugins/Runtime/OpenXR/Source/OpenXRInput/Public/OpenXRInputSettings.h
deleted file mode 100644
diff --git a/Templates/TP_VirtualRealityBP/Config/DefaultInput.ini b/Templates/TP_VirtualRealityBP/Config/DefaultInput.ini
--- a/Templates/TP_VirtualRealityBP/Config/DefaultInput.ini
+++ b/Templates/TP_VirtualRealityBP/Config/DefaultInput.ini
@@ -107,7 +107,10 @@ DefaultTouchInterface=None
 -ConsoleKeys=Tilde
 +ConsoleKeys=Tilde
 
-[/Script/OpenXRInput.OpenXRInputSettings]
-MappableInputConfig=None
-InputMappingContexts=("/Game/VRTemplate/Input/IMC_Default.IMC_Default","/Game/VRTemplate/Input/IMC_Menu.IMC_Menu","/Game/VRTemplate/Input/IMC_Weapon_Left.IMC_Weapon_Left","/Game/VRTemplate/Input/IMC_Weapon_Right.IMC_Weapon_Right","/Game/VRTemplate/Input/IMC_Hands.IMC_Hands")
+[/Script/EnhancedInput.EnhancedInputDeveloperSettings]
++DefaultMappingContexts=(InputMappingContext="/Game/VRTemplate/Input/IMC_Default.IMC_Default",Priority=0,bAddImmediately=True,bRegisterWithUserSettings=False)
++DefaultMappingContexts=(InputMappingContext="/Game/VRTemplate/Input/IMC_Hands.IMC_Hands",Priority=0,bAddImmediately=True,bRegisterWithUserSettings=False)
++DefaultMappingContexts=(InputMappingContext="/Game/VRTemplate/Input/IMC_Menu.IMC_Menu",Priority=0,bAddImmediately=True,bRegisterWithUserSettings=False)
++DefaultMappingContexts=(InputMappingContext="/Game/VRTemplate/Input/IMC_Weapon_Left.IMC_Weapon_Left",Priority=0,bAddImmediately=True,bRegisterWithUserSettings=False)
++DefaultMappingContexts=(InputMappingContext="/Game/VRTemplate/Input/IMC_Weapon_Right.IMC_Weapon_Right",Priority=0,bAddImmediately=True,bRegisterWithUserSettings=False)
 
