Fix for null-pointer dereference crash in OpenXRInput.cpp when Enhanced Input -> Default Mapping Contexts contains an entry with an Input Mapping Context of none.

diff --git a/Engine/Plugins/Runtime/OpenXR/Source/OpenXRInput/Private/OpenXRInput.cpp b/Engine/Plugins/Runtime/OpenXR/Source/OpenXRInput/Private/OpenXRInput.cpp
--- a/Engine/Plugins/Runtime/OpenXR/Source/OpenXRInput/Private/OpenXRInput.cpp
+++ b/Engine/Plugins/Runtime/OpenXR/Source/OpenXRInput/Private/OpenXRInput.cpp
@@ -408,8 +408,15 @@ bool FOpenXRInputPlugin::FOpenXRInput::BuildActions(XrSession Session)
 	{
 		for (const auto& Context : InputSettings->DefaultMappingContexts)
 		{
-			TStrongObjectPtr<const UInputMappingContext> Obj(Context.InputMappingContext.LoadSynchronous());
-			InputMappingContextToPriorityMap.Add(Obj, Context.Priority);
+			if (Context.InputMappingContext)
+			{
+				TStrongObjectPtr<const UInputMappingContext> Obj(Context.InputMappingContext.LoadSynchronous());
+				InputMappingContextToPriorityMap.Add(Obj, Context.Priority);
+			}
+			else
+			{
+				UE_LOG(LogHMD, Warning, TEXT("Default Mapping Contexts contains an Input Mapping Context set to \"None\", ignoring while building OpenXR actions."));
+			}
 		}
 	}
 
