Deprecate:

diff --git a/Engine/Plugins/Runtime/OpenXR/Source/OpenXRHMD/Private/OpenXRHMD.cpp b/Engine/Plugins/Runtime/OpenXR/Source/OpenXRHMD/Private/OpenXRHMD.cpp
--- a/Engine/Plugins/Runtime/OpenXR/Source/OpenXRHMD/Private/OpenXRHMD.cpp
+++ b/Engine/Plugins/Runtime/OpenXR/Source/OpenXRHMD/Private/OpenXRHMD.cpp
@@ -2357,7 +2357,7 @@ int32 FOpenXRHMD::AcquireColorTexture()
 	return 0;
 }
 
-bool FOpenXRHMD::AllocateRenderTargetTextures(uint32 SizeX, uint32 SizeY, uint8 Format, uint32 NumMips, ETextureCreateFlags Flags, ETextureCreateFlags TargetableTextureFlags, TArray<FTexture2DRHIRef>& OutTargetableTextures, TArray<FTexture2DRHIRef>& OutShaderResourceTextures, uint32 NumSamples)
+bool FOpenXRHMD::AllocateRenderTargetTextures(uint32 SizeX, uint32 SizeY, uint8 Format, uint32 NumMips, ETextureCreateFlags Flags, ETextureCreateFlags TargetableTextureFlags, TArray<FTextureRHIRef>& OutTargetableTextures, TArray<FTextureRHIRef>& OutShaderResourceTextures, uint32 NumSamples)
 {
 	check(IsInRenderingThread());
 
@@ -2409,7 +2409,7 @@ bool FOpenXRHMD::AllocateRenderTargetTextures(uint32 SizeX, uint32 SizeY, uint8
 
 	uint8 ActualFormat = Format;
 	FXRSwapChainPtr& Swapchain = PipelinedLayerStateRendering.ColorSwapchain;
-	const FRHITexture2D* const SwapchainTexture = Swapchain == nullptr ? nullptr : Swapchain->GetTexture2DArray() ? Swapchain->GetTexture2DArray() : Swapchain->GetTexture2D();
+	const FRHITexture* const SwapchainTexture = Swapchain == nullptr ? nullptr : Swapchain->GetTexture2DArray() ? Swapchain->GetTexture2DArray() : Swapchain->GetTexture2D();
 	if (Swapchain == nullptr || SwapchainTexture == nullptr || Format != LastRequestedColorSwapchainFormat || SwapchainTexture->GetSizeX() != SizeX || SwapchainTexture->GetSizeY() != SizeY)
 	{
 		ensureMsgf(NumSamples == 1, TEXT("OpenXR supports MSAA swapchains, but engine logic expects the swapchain target to be 1x."));
@@ -2443,7 +2443,7 @@ bool FOpenXRHMD::AllocateRenderTargetTextures(uint32 SizeX, uint32 SizeY, uint8
 	{
 		// If we have at least two native layers, use non-background layer to render the composited image of all the emulated face locked layers.
 		FXRSwapChainPtr& EmulationSwapchain = PipelinedLayerStateRendering.EmulatedLayerState.EmulationSwapchain;
-		const FRHITexture2D* const EmulationSwapchainTexture = EmulationSwapchain == nullptr ? nullptr : EmulationSwapchain->GetTexture2DArray() ? EmulationSwapchain->GetTexture2DArray() : EmulationSwapchain->GetTexture2D();
+		const FRHITexture* const EmulationSwapchainTexture = EmulationSwapchain == nullptr ? nullptr : EmulationSwapchain->GetTexture2DArray() ? EmulationSwapchain->GetTexture2DArray() : EmulationSwapchain->GetTexture2D();
 		if (EmulationSwapchain == nullptr || EmulationSwapchainTexture == nullptr || EmulationSwapchainTexture->GetSizeX() != SizeX || EmulationSwapchainTexture->GetSizeY() != SizeY)
 		{
 			const ETextureCreateFlags EmulationCreateFlags = TexCreate_Dynamic | TexCreate_ShaderResource | TexCreate_RenderTargetable;
@@ -2478,7 +2478,7 @@ void FOpenXRHMD::AllocateDepthTextureInternal(uint32 SizeX, uint32 SizeY, uint32
 	}
 
 	FXRSwapChainPtr& DepthSwapchain = PipelinedLayerStateRendering.DepthSwapchain;
-	const FRHITexture2D* const DepthSwapchainTexture = DepthSwapchain == nullptr ? nullptr : DepthSwapchain->GetTexture2DArray() ? DepthSwapchain->GetTexture2DArray() : DepthSwapchain->GetTexture2D();
+	const FRHITexture* const DepthSwapchainTexture = DepthSwapchain == nullptr ? nullptr : DepthSwapchain->GetTexture2DArray() ? DepthSwapchain->GetTexture2DArray() : DepthSwapchain->GetTexture2D();
 	if (DepthSwapchain == nullptr || DepthSwapchainTexture == nullptr ||
 		DepthSwapchainTexture->GetSizeX() != SizeX || DepthSwapchainTexture->GetSizeY() != SizeY || DepthSwapchainTexture->GetDesc().ArraySize != InArraySize)
 	{
@@ -2514,7 +2514,7 @@ void FOpenXRHMD::AllocateDepthTextureInternal(uint32 SizeX, uint32 SizeY, uint32
 }
 
 // TODO: in the future, we can rename the interface to GetDepthTexture because allocate could happen in AllocateRenderTargetTexture
-bool FOpenXRHMD::AllocateDepthTexture(uint32 Index, uint32 SizeX, uint32 SizeY, uint8 Format, uint32 NumMips, ETextureCreateFlags Flags, ETextureCreateFlags TargetableTextureFlags, FTexture2DRHIRef& OutTargetableTexture, FTexture2DRHIRef& OutShaderResourceTexture, uint32 NumSamples)
+bool FOpenXRHMD::AllocateDepthTexture(uint32 Index, uint32 SizeX, uint32 SizeY, uint8 Format, uint32 NumMips, ETextureCreateFlags Flags, ETextureCreateFlags TargetableTextureFlags, FTextureRHIRef& OutTargetableTexture, FTextureRHIRef& OutShaderResourceTexture, uint32 NumSamples)
 {
 	check(IsInRenderingThread());
 
@@ -2538,7 +2538,7 @@ bool FOpenXRHMD::AllocateDepthTexture(uint32 Index, uint32 SizeX, uint32 SizeY,
 		DepthSwapchain->IncrementSwapChainIndex_RHIThread();
 	}
 
-	const FRHITexture2D* const DepthSwapchainTexture = DepthSwapchain->GetTexture2DArray() ? DepthSwapchain->GetTexture2DArray() : DepthSwapchain->GetTexture2D();
+	const FRHITexture* const DepthSwapchainTexture = DepthSwapchain->GetTexture2DArray() ? DepthSwapchain->GetTexture2DArray() : DepthSwapchain->GetTexture2D();
 	const FRHITextureDesc& DepthSwapchainDesc = DepthSwapchainTexture->GetDesc();
 
 	if (SizeX != DepthSwapchainDesc.Extent.X || SizeY != DepthSwapchainDesc.Extent.Y)
@@ -2557,7 +2557,7 @@ bool FOpenXRHMD::AllocateDepthTexture(uint32 Index, uint32 SizeX, uint32 SizeY,
 
 	LastRequestedDepthSwapchainFormat = Format;
 
-	OutTargetableTexture = OutShaderResourceTexture = (FTexture2DRHIRef&)PipelinedLayerStateRendering.DepthSwapchain->GetTextureRef();
+	OutTargetableTexture = OutShaderResourceTexture = (FTextureRHIRef&)PipelinedLayerStateRendering.DepthSwapchain->GetTextureRef();
 
 	PipelinedLayerStateRendering.LayerStateFlags |= EOpenXRLayerStateFlags::SubmitDepthLayer;
 
@@ -2566,7 +2566,7 @@ bool FOpenXRHMD::AllocateDepthTexture(uint32 Index, uint32 SizeX, uint32 SizeY,
 
 void CreateNativeLayerSwapchain(FOpenXRLayer& Layer, TRefCountPtr<FOpenXRRenderBridge>& RenderBridge, XrSession Session)
 {
-	auto CreateSwapchain = [&](FRHITexture2D* Texture, ETextureCreateFlags Flags)
+	auto CreateSwapchain = [&](FRHITexture* Texture, ETextureCreateFlags Flags)
 	{
 		uint8 UnusedActualFormat = 0;
 		return RenderBridge->CreateSwapchain(Session,
@@ -2586,13 +2586,13 @@ void CreateNativeLayerSwapchain(FOpenXRLayer& Layer, TRefCountPtr<FOpenXRRenderB
 
 	if (Layer.NeedReallocateRightTexture())
 	{
-		FRHITexture2D* Texture = Layer.Desc.Texture->GetTexture2D();
+		FRHITexture* Texture = Layer.Desc.Texture->GetTexture2D();
 		Layer.RightEye.SetSwapchain(CreateSwapchain(Texture, Flags), Texture->GetSizeXY());
 	}
 
 	if (Layer.NeedReallocateLeftTexture())
 	{
-		FRHITexture2D* Texture = Layer.Desc.LeftTexture->GetTexture2D();
+		FRHITexture* Texture = Layer.Desc.LeftTexture->GetTexture2D();
 		Layer.LeftEye.SetSwapchain(CreateSwapchain(Texture, Flags), Texture->GetSizeXY());
 	}
 }
@@ -2770,7 +2770,7 @@ void FOpenXRHMD::UpdateLayerSwapchainTexture(const FOpenXRLayer& Layer, FRHIComm
 	{
 		if (Layer.RightEye.bUpdateTexture && bIsRunning)
 		{
-			FRHITexture2D* SrcTexture = Layer.Desc.Texture->GetTexture2D();
+			FRHITexture* SrcTexture = Layer.Desc.Texture->GetTexture2D();
 			FIntRect DstRect(FIntPoint(0, 0), Layer.RightEye.SwapchainSize.IntPoint());
 			CopyTexture_RenderThread(RHICmdList, SrcTexture, FIntRect(), Layer.RightEye.Swapchain, DstRect, false, SrcTextureCopyModifier);
 		}
@@ -2780,7 +2780,7 @@ void FOpenXRHMD::UpdateLayerSwapchainTexture(const FOpenXRLayer& Layer, FRHIComm
 	{
 		if (Layer.LeftEye.bUpdateTexture && bIsRunning)
 		{
-			FRHITexture2D* SrcTexture = Layer.Desc.LeftTexture->GetTexture2D();
+			FRHITexture* SrcTexture = Layer.Desc.LeftTexture->GetTexture2D();
 			FIntRect DstRect(FIntPoint(0, 0), Layer.LeftEye.SwapchainSize.IntPoint());
 			CopyTexture_RenderThread(RHICmdList, SrcTexture, FIntRect(), Layer.LeftEye.Swapchain, DstRect, false, SrcTextureCopyModifier);
 		}
@@ -2810,7 +2810,7 @@ void FOpenXRHMD::DrawEmulatedFaceLockedLayers_RenderThread(FRDGBuilder& GraphBui
 	AddPass(GraphBuilder, RDG_EVENT_NAME("OpenXREmulatedFaceLockedLayerRender"), [this, &InView](FRHICommandListImmediate& RHICmdList)
 	{
 		FXRSwapChainPtr EmulationSwapchain = PipelinedLayerStateRendering.EmulatedLayerState.EmulationSwapchain;
-		FTexture2DRHIRef RenderTarget = EmulationSwapchain->GetTextureRef();
+		FTextureRHIRef RenderTarget = EmulationSwapchain->GetTextureRef();
 
 		FDefaultStereoLayers_LayerRenderParams RenderParams;
 		FRHIRenderPassInfo RPInfo = SetupEmulatedLayersRenderPass(RHICmdList, InView, EmulatedFaceLockedLayers, RenderTarget, RenderParams);
@@ -2832,7 +2832,7 @@ void FOpenXRHMD::DrawBackgroundCompositedEmulatedLayers_RenderThread(FRDGBuilder
 	// Partially borrowed from FDefaultStereoLayers
 	AddPass(GraphBuilder, RDG_EVENT_NAME("OpenXREmulatedLayerRender"), [this, &InView](FRHICommandListImmediate& RHICmdList)
 	{
-		FTexture2DRHIRef RenderTarget = InView.Family->RenderTarget->GetRenderTargetTexture();
+		FTextureRHIRef RenderTarget = InView.Family->RenderTarget->GetRenderTargetTexture();
 
 		FDefaultStereoLayers_LayerRenderParams RenderParams;
 		FRHIRenderPassInfo RPInfo = SetupEmulatedLayersRenderPass(RHICmdList, InView, BackgroundCompositedEmulatedLayers, RenderTarget, RenderParams);
@@ -2851,7 +2851,7 @@ void FOpenXRHMD::DrawBackgroundCompositedEmulatedLayers_RenderThread(FRDGBuilder
 	});
 }
 
-FRHIRenderPassInfo FOpenXRHMD::SetupEmulatedLayersRenderPass(FRHICommandListImmediate& RHICmdList, const FSceneView& InView, TArray<IStereoLayers::FLayerDesc>& Layers, FTexture2DRHIRef RenderTarget, FDefaultStereoLayers_LayerRenderParams& OutRenderParams)
+FRHIRenderPassInfo FOpenXRHMD::SetupEmulatedLayersRenderPass(FRHICommandListImmediate& RHICmdList, const FSceneView& InView, TArray<IStereoLayers::FLayerDesc>& Layers, FTextureRHIRef RenderTarget, FDefaultStereoLayers_LayerRenderParams& OutRenderParams)
 {
 	OutRenderParams = CalculateEmulatedLayerRenderParams(InView);
 	TArray<FRHITransitionInfo, TInlineAllocator<16>> Infos;
@@ -3667,7 +3667,7 @@ FIntPoint FOpenXRHMD::GetIdealRenderTargetSize() const
 	return Size;
 }
 
-FIntRect FOpenXRHMD::GetFullFlatEyeRect_RenderThread(FTexture2DRHIRef EyeTexture) const
+FIntRect FOpenXRHMD::GetFullFlatEyeRect_RenderThread(FTextureRHIRef EyeTexture) const
 {
 	FVector2D SrcNormRectMin(0.05f, 0.2f);
 	// with MMV, each eye occupies the whole RT layer, so we don't need to limit the source rect to the left half of the RT.
@@ -3832,7 +3832,7 @@ IMPLEMENT_SHADER_TYPE(, FDisplayMappingPS, TEXT("/Engine/Private/DisplayMappingP
 // * transparent export just writes thru RGBA, no clears needed
 // * premultiplied alpha blend needs custom blend, no clears needed
 
-void FOpenXRHMD::CopyTexture_RenderThread(FRHICommandListImmediate& RHICmdList, FRHITexture2D* SrcTexture, FIntRect SrcRect, FRHITexture2D* DstTexture, FIntRect DstRect, 
+void FOpenXRHMD::CopyTexture_RenderThread(FRHICommandListImmediate& RHICmdList, FRHITexture* SrcTexture, FIntRect SrcRect, FRHITexture* DstTexture, FIntRect DstRect, 
 											bool bClearBlack, ERenderTargetActions RTAction, ERHIAccess FinalDstAccess, ETextureCopyBlendModifier SrcTextureCopyModifier) const
 {
 	check(IsInRenderingThread());
@@ -4018,7 +4018,7 @@ void FOpenXRHMD::CopyTexture_RenderThread(FRHICommandListImmediate& RHICmdList,
 	RHICmdList.Transition(FRHITransitionInfo(DstTexture, ERHIAccess::RTV, FinalDstAccess));
 }
 
-void FOpenXRHMD::CopyTexture_RenderThread(FRHICommandListImmediate& RHICmdList, FRHITexture2D* SrcTexture, FIntRect SrcRect, const FXRSwapChainPtr& DstSwapChain, FIntRect DstRect, bool bClearBlack, ETextureCopyBlendModifier SrcTextureCopyModifier) const
+void FOpenXRHMD::CopyTexture_RenderThread(FRHICommandListImmediate& RHICmdList, FRHITexture* SrcTexture, FIntRect SrcRect, const FXRSwapChainPtr& DstSwapChain, FIntRect DstRect, bool bClearBlack, ETextureCopyBlendModifier SrcTextureCopyModifier) const
 {
 	RHICmdList.EnqueueLambda([DstSwapChain](FRHICommandListImmediate& InRHICmdList)
 	{
@@ -4027,7 +4027,7 @@ void FOpenXRHMD::CopyTexture_RenderThread(FRHICommandListImmediate& RHICmdList,
 	});
 
 	// Now that we've enqueued the swapchain wait we can add the commands to do the actual texture copy
-	FRHITexture2D* const DstTexture = DstSwapChain->GetTexture2DArray() ? DstSwapChain->GetTexture2DArray() : DstSwapChain->GetTexture2D();
+	FRHITexture* const DstTexture = DstSwapChain->GetTexture2DArray() ? DstSwapChain->GetTexture2DArray() : DstSwapChain->GetTexture2D();
 	CopyTexture_RenderThread(RHICmdList, SrcTexture, SrcRect, DstTexture, DstRect, bClearBlack, ERenderTargetActions::Clear_Store, ERHIAccess::SRVMask, SrcTextureCopyModifier);
 
 	// Enqueue a command to release the image after the copy is done
@@ -4037,7 +4037,7 @@ void FOpenXRHMD::CopyTexture_RenderThread(FRHICommandListImmediate& RHICmdList,
 	});
 }
 
-void FOpenXRHMD::CopyTexture_RenderThread(FRHICommandListImmediate& RHICmdList, FRHITexture2D* SrcTexture, FIntRect SrcRect, FRHITexture2D* DstTexture, FIntRect DstRect, bool bClearBlack, bool bNoAlpha) const
+void FOpenXRHMD::CopyTexture_RenderThread(FRHICommandListImmediate& RHICmdList, FRHITexture* SrcTexture, FIntRect SrcRect, FRHITexture* DstTexture, FIntRect DstRect, bool bClearBlack, bool bNoAlpha) const
 {
 	// This call only comes from the spectator screen so we expect alpha to be premultiplied.
 	const ETextureCopyBlendModifier SrcTextureCopyModifier = bNoAlpha ? ETextureCopyBlendModifier::Opaque : ETextureCopyBlendModifier::PremultipliedAlphaBlend;
@@ -4049,7 +4049,7 @@ void FOpenXRHMD::RenderTexture_RenderThread(class FRHICommandListImmediate& RHIC
 	if (SpectatorScreenController)
 	{
 		const bool bShouldPassLayersTexture = EnumHasAnyFlags(PipelinedLayerStateRendering.LayerStateFlags, EOpenXRLayerStateFlags::SubmitEmulatedFaceLockedLayer) && !CVarOpenXRDoNotCopyEmulatedLayersToSpectatorScreen.GetValueOnRenderThread();
-		const FTexture2DRHIRef LayersTexture = bShouldPassLayersTexture ? PipelinedLayerStateRendering.EmulatedLayerState.EmulationSwapchain->GetTextureRef() : nullptr;
+		const FTextureRHIRef LayersTexture = bShouldPassLayersTexture ? PipelinedLayerStateRendering.EmulatedLayerState.EmulationSwapchain->GetTextureRef() : nullptr;
 		SpectatorScreenController->RenderSpectatorScreen_RenderThread(RHICmdList, BackBuffer, SrcTexture, LayersTexture, WindowSize);
 	}
 }
diff --git a/Engine/Plugins/Runtime/OpenXR/Source/OpenXRHMD/Private/OpenXRHMD.h b/Engine/Plugins/Runtime/OpenXR/Source/OpenXRHMD/Private/OpenXRHMD.h
--- a/Engine/Plugins/Runtime/OpenXR/Source/OpenXRHMD/Private/OpenXRHMD.h
+++ b/Engine/Plugins/Runtime/OpenXR/Source/OpenXRHMD/Private/OpenXRHMD.h
@@ -282,10 +282,10 @@ protected:
 	void EnumerateViews(FPipelinedFrameState& PipelineState);
 	void LocateViews(FPipelinedFrameState& PipelinedState, bool ResizeViewsArray = false);
 
-	void CopyTexture_RenderThread(FRHICommandListImmediate& RHICmdList, FRHITexture2D* SrcTexture, FIntRect SrcRect, FRHITexture2D* DstTexture, FIntRect DstRect, 
+	void CopyTexture_RenderThread(FRHICommandListImmediate& RHICmdList, FRHITexture* SrcTexture, FIntRect SrcRect, FRHITexture* DstTexture, FIntRect DstRect, 
 								  bool bClearBlack, ERenderTargetActions RTAction, ERHIAccess FinalDstAccess, ETextureCopyBlendModifier SrcTextureCopyModifier) const;
 	
-	void CopyTexture_RenderThread(FRHICommandListImmediate& RHICmdList, FRHITexture2D* SrcTexture, FIntRect SrcRect, const FXRSwapChainPtr& DstSwapChain, FIntRect DstRect, bool bClearBlack, ETextureCopyBlendModifier SrcTextureCopyModifier) const;
+	void CopyTexture_RenderThread(FRHICommandListImmediate& RHICmdList, FRHITexture* SrcTexture, FIntRect SrcRect, const FXRSwapChainPtr& DstSwapChain, FIntRect DstRect, bool bClearBlack, ETextureCopyBlendModifier SrcTextureCopyModifier) const;
 
 	void AllocateDepthTextureInternal(uint32 SizeX, uint32 SizeY, uint32 NumSamples, uint32 ArraySize);
 
@@ -317,8 +317,8 @@ public:
 	virtual void SetPixelDensity(const float NewDensity) override;
 	virtual FIntPoint GetIdealRenderTargetSize() const override;
 	virtual bool GetHMDDistortionEnabled(EShadingPath ShadingPath) const override { return false; }
-	virtual FIntRect GetFullFlatEyeRect_RenderThread(FTexture2DRHIRef EyeTexture) const override;
-	virtual void CopyTexture_RenderThread(FRHICommandListImmediate& RHICmdList, FRHITexture2D* SrcTexture, FIntRect SrcRect, FRHITexture2D* DstTexture, FIntRect DstRect, bool bClearBlack, bool bNoAlpha) const override;
+	virtual FIntRect GetFullFlatEyeRect_RenderThread(FTextureRHIRef EyeTexture) const override;
+	virtual void CopyTexture_RenderThread(FRHICommandListImmediate& RHICmdList, FRHITexture* SrcTexture, FIntRect SrcRect, FRHITexture* DstTexture, FIntRect DstRect, bool bClearBlack, bool bNoAlpha) const override;
 	virtual bool HasHiddenAreaMesh() const override final;
 	virtual bool HasVisibleAreaMesh() const override final;
 	virtual void DrawHiddenAreaMesh(class FRHICommandList& RHICmdList, int32 ViewIndex) const override final;
@@ -360,9 +360,9 @@ public:
 	/** IStereoRenderTargetManager */
 	virtual bool ShouldUseSeparateRenderTarget() const override { return IsStereoEnabled() && RenderBridge.IsValid(); }
 	virtual void CalculateRenderTargetSize(const FViewport& Viewport, uint32& InOutSizeX, uint32& InOutSizeY) override;
-	virtual bool AllocateRenderTargetTextures(uint32 SizeX, uint32 SizeY, uint8 Format, uint32 NumLayers, ETextureCreateFlags Flags, ETextureCreateFlags TargetableTextureFlags, TArray<FTexture2DRHIRef>& OutTargetableTextures, TArray<FTexture2DRHIRef>& OutShaderResourceTextures, uint32 NumSamples = 1) override;
+	virtual bool AllocateRenderTargetTextures(uint32 SizeX, uint32 SizeY, uint8 Format, uint32 NumLayers, ETextureCreateFlags Flags, ETextureCreateFlags TargetableTextureFlags, TArray<FTextureRHIRef>& OutTargetableTextures, TArray<FTextureRHIRef>& OutShaderResourceTextures, uint32 NumSamples = 1) override;
 	virtual int32 AcquireColorTexture() override final;
-	virtual bool AllocateDepthTexture(uint32 Index, uint32 SizeX, uint32 SizeY, uint8 Format, uint32 NumMips, ETextureCreateFlags InTexFlags, ETextureCreateFlags TargetableTextureFlags, FTexture2DRHIRef& OutTargetableTexture, FTexture2DRHIRef& OutShaderResourceTexture, uint32 NumSamples = 1) override final;
+	virtual bool AllocateDepthTexture(uint32 Index, uint32 SizeX, uint32 SizeY, uint8 Format, uint32 NumMips, ETextureCreateFlags InTexFlags, ETextureCreateFlags TargetableTextureFlags, FTextureRHIRef& OutTargetableTexture, FTextureRHIRef& OutShaderResourceTexture, uint32 NumSamples = 1) override final;
 	virtual bool ReconfigureForShaderPlatform(EShaderPlatform NewShaderPlatform) override;
 	virtual EPixelFormat GetActualColorSwapchainFormat() const override { return static_cast<EPixelFormat>(LastActualColorSwapchainFormat); }
 
@@ -426,7 +426,7 @@ private:
 
 	TArray<XrEnvironmentBlendMode> RetrieveEnvironmentBlendModes() const;
 	FDefaultStereoLayers_LayerRenderParams CalculateEmulatedLayerRenderParams(const FSceneView& InView);
-	FRHIRenderPassInfo SetupEmulatedLayersRenderPass(FRHICommandListImmediate& RHICmdList, const FSceneView& InView, TArray<IStereoLayers::FLayerDesc>& Layers, FTexture2DRHIRef RenderTarget, FDefaultStereoLayers_LayerRenderParams& OutRenderParams);
+	FRHIRenderPassInfo SetupEmulatedLayersRenderPass(FRHICommandListImmediate& RHICmdList, const FSceneView& InView, TArray<IStereoLayers::FLayerDesc>& Layers, FTextureRHIRef RenderTarget, FDefaultStereoLayers_LayerRenderParams& OutRenderParams);
 	bool IsEmulatingStereoLayers();
 	
 	void UpdateLayerSwapchainTexture(const FOpenXRLayer& Layer, FRHICommandListImmediate& RHICmdList);
diff --git a/Engine/Plugins/Runtime/OpenXR/Source/OpenXRHMD/Private/OpenXRHMD_Layer.cpp b/Engine/Plugins/Runtime/OpenXR/Source/OpenXRHMD/Private/OpenXRHMD_Layer.cpp
--- a/Engine/Plugins/Runtime/OpenXR/Source/OpenXRHMD/Private/OpenXRHMD_Layer.cpp
+++ b/Engine/Plugins/Runtime/OpenXR/Source/OpenXRHMD/Private/OpenXRHMD_Layer.cpp
@@ -13,7 +13,7 @@ bool FOpenXRLayer::NeedReallocateRightTexture()
 		return false;
 	}
 
-	FRHITexture2D* Texture = Desc.Texture->GetTexture2D();
+	FRHITexture* Texture = Desc.Texture->GetTexture2D();
 	if (!Texture)
 	{
 		return false;
@@ -34,7 +34,7 @@ bool FOpenXRLayer::NeedReallocateLeftTexture()
 		return false;
 	}
 
-	FRHITexture2D* Texture = Desc.LeftTexture->GetTexture2D();
+	FRHITexture* Texture = Desc.LeftTexture->GetTexture2D();
 	if (!Texture)
 	{
 		return false;
diff --git a/Engine/Plugins/Runtime/OpenXR/Source/OpenXRHMD/Private/OpenXRHMD_Swapchain.cpp b/Engine/Plugins/Runtime/OpenXR/Source/OpenXRHMD/Private/OpenXRHMD_Swapchain.cpp
--- a/Engine/Plugins/Runtime/OpenXR/Source/OpenXRHMD/Private/OpenXRHMD_Swapchain.cpp
+++ b/Engine/Plugins/Runtime/OpenXR/Source/OpenXRHMD/Private/OpenXRHMD_Swapchain.cpp
@@ -374,7 +374,7 @@ FXRSwapChainPtr CreateSwapchain_OpenGL(XrSession InSession, uint8 Format, uint8&
 	TArray<XrSwapchainImageOpenGLKHR> Images = EnumerateImages<XrSwapchainImageOpenGLKHR>(Swapchain, XR_TYPE_SWAPCHAIN_IMAGE_OPENGL_KHR);
 	for (const auto& Image : Images)
 	{
-		FTexture2DRHIRef NewTexture = (ArraySize > 1) ?
+		FTextureRHIRef NewTexture = (ArraySize > 1) ?
 			DynamicRHI->RHICreateTexture2DArrayFromResource(GPixelFormats[Format].UnrealFormat, SizeX, SizeY, ArraySize, NumMips, NumSamples, 1, ClearValueBinding, Image.image, CreateFlags) :
 			DynamicRHI->RHICreateTexture2DFromResource(GPixelFormats[Format].UnrealFormat, SizeX, SizeY, NumMips, NumSamples, 1, ClearValueBinding, Image.image, CreateFlags);
 		TextureChain.Add(NewTexture.GetReference());
@@ -414,7 +414,7 @@ FXRSwapChainPtr CreateSwapchain_OpenGLES(XrSession InSession, uint8 Format, uint
 	TArray<XrSwapchainImageOpenGLESKHR> Images = EnumerateImages<XrSwapchainImageOpenGLESKHR>(Swapchain, XR_TYPE_SWAPCHAIN_IMAGE_OPENGL_ES_KHR);
 	for (const auto& Image : Images)
 	{
-		FTexture2DRHIRef NewTexture = (ArraySize > 1) ?
+		FTextureRHIRef NewTexture = (ArraySize > 1) ?
 			DynamicRHI->RHICreateTexture2DArrayFromResource(GPixelFormats[Format].UnrealFormat, SizeX, SizeY, ArraySize, NumMips, NumSamples, 1, ClearValueBinding, Image.image, CreateFlags) :
 			DynamicRHI->RHICreateTexture2DFromResource(GPixelFormats[Format].UnrealFormat, SizeX, SizeY, NumMips, NumSamples, 1, ClearValueBinding, Image.image, CreateFlags);
 		TextureChain.Add(NewTexture.GetReference());
diff --git a/Engine/Plugins/Runtime/WarpUtils/Source/PFMExporter/Private/PFMExporterMesh.h b/Engine/Plugins/Runtime/WarpUtils/Source/PFMExporter/Private/PFMExporterMesh.h
--- a/Engine/Plugins/Runtime/WarpUtils/Source/PFMExporter/Private/PFMExporterMesh.h
+++ b/Engine/Plugins/Runtime/WarpUtils/Source/PFMExporter/Private/PFMExporterMesh.h
@@ -27,14 +27,14 @@ public:
 		return FIntRect(FIntPoint(0, 0), FIntPoint(DimWidth, DimHeight));
 	}
 
-	inline FTexture2DRHIRef GetTargetableTexture() const
+	inline FTextureRHIRef GetTargetableTexture() const
 	{
 		return RenderTargetTexture;
 	}
 
 private:
-	FTexture2DRHIRef RenderTargetTexture;
-	FTexture2DRHIRef ShaderResourceTexture;
+	FTextureRHIRef RenderTargetTexture;
+	FTextureRHIRef ShaderResourceTexture;
 };
 
 
diff --git a/Engine/Plugins/Runtime/WebMMoviePlayer/Source/WebMMoviePlayer/Private/WebMMovieStreamer.cpp b/Engine/Plugins/Runtime/WebMMoviePlayer/Source/WebMMoviePlayer/Private/WebMMovieStreamer.cpp
--- a/Engine/Plugins/Runtime/WebMMoviePlayer/Source/WebMMoviePlayer/Private/WebMMovieStreamer.cpp
+++ b/Engine/Plugins/Runtime/WebMMoviePlayer/Source/WebMMoviePlayer/Private/WebMMovieStreamer.cpp
@@ -57,7 +57,7 @@ void FWebMMovieStreamer::Cleanup()
 	AudioBackend->ShutdownPlatform();
 }
 
-FTexture2DRHIRef FWebMMovieStreamer::GetTexture()
+FTextureRHIRef FWebMMovieStreamer::GetTexture()
 {
 	return SlateVideoTexture.IsValid() ? SlateVideoTexture->GetRHIRef() : nullptr;
 }
diff --git a/Engine/Plugins/Runtime/WebMMoviePlayer/Source/WebMMoviePlayer/Private/WebMMovieStreamer.h b/Engine/Plugins/Runtime/WebMMoviePlayer/Source/WebMMoviePlayer/Private/WebMMovieStreamer.h
--- a/Engine/Plugins/Runtime/WebMMoviePlayer/Source/WebMMoviePlayer/Private/WebMMovieStreamer.h
+++ b/Engine/Plugins/Runtime/WebMMoviePlayer/Source/WebMMoviePlayer/Private/WebMMovieStreamer.h
@@ -36,7 +36,7 @@ public:
 	virtual FString GetMovieName() override;
 	virtual bool IsLastMovieInPlaylist() override;
 	virtual void Cleanup() override;
-	virtual FTexture2DRHIRef GetTexture() override;
+	virtual FTextureRHIRef GetTexture() override;
 
 	FOnCurrentMovieClipFinished OnCurrentMovieClipFinishedDelegate;
 	virtual FOnCurrentMovieClipFinished& OnCurrentMovieClipFinished() override { return OnCurrentMovieClipFinishedDelegate; }
diff --git a/Engine/Plugins/Runtime/WindowsMoviePlayer/Source/WindowsMoviePlayer/Private/WindowsMovieStreamer.h b/Engine/Plugins/Runtime/WindowsMoviePlayer/Source/WindowsMoviePlayer/Private/WindowsMovieStreamer.h
--- a/Engine/Plugins/Runtime/WindowsMoviePlayer/Source/WindowsMoviePlayer/Private/WindowsMovieStreamer.h
+++ b/Engine/Plugins/Runtime/WindowsMoviePlayer/Source/WindowsMoviePlayer/Private/WindowsMovieStreamer.h
@@ -43,7 +43,7 @@ public:
 	FOnCurrentMovieClipFinished OnCurrentMovieClipFinishedDelegate;
 	virtual FOnCurrentMovieClipFinished& OnCurrentMovieClipFinished() override { return OnCurrentMovieClipFinishedDelegate; }
 
-	virtual FTexture2DRHIRef GetTexture() override { return Texture.IsValid() ? Texture->GetRHIRef() : nullptr; }
+	virtual FTextureRHIRef GetTexture() override { return Texture.IsValid() ? Texture->GetRHIRef() : nullptr; }
 
 private:
 	/** Opens up the next movie in the movie path queue */
@@ -77,7 +77,7 @@ private:
 	TArray<TSharedPtr<FSlateTexture2DRHIRef, ESPMode::ThreadSafe>> TextureFreeList;
 
 	/** Input texture for conversion shader */
-	TRefCountPtr<FRHITexture2D> InputTarget;
+	TRefCountPtr<FRHITexture> InputTarget;
 
 	/** The video player and sample grabber for use of Media Foundation */
 	class FVideoPlayer* VideoPlayer;
diff --git a/Engine/Plugins/Runtime/nDisplay/Source/DisplayCluster/Private/Render/DisplayDevice/Proxy/DisplayClusterDisplayDeviceProxy_OpenColorIO.cpp b/Engine/Plugins/Runtime/nDisplay/Source/DisplayCluster/Private/Render/DisplayDevice/Proxy/DisplayClusterDisplayDeviceProxy_OpenColorIO.cpp
--- a/Engine/Plugins/Runtime/nDisplay/Source/DisplayCluster/Private/Render/DisplayDevice/Proxy/DisplayClusterDisplayDeviceProxy_OpenColorIO.cpp
+++ b/Engine/Plugins/Runtime/nDisplay/Source/DisplayCluster/Private/Render/DisplayDevice/Proxy/DisplayClusterDisplayDeviceProxy_OpenColorIO.cpp
@@ -35,7 +35,7 @@ bool FDisplayClusterDisplayDeviceProxy_OpenColorIO::HasFinalPass_RenderThread()
 }
 
 bool FDisplayClusterDisplayDeviceProxy_OpenColorIO::AddFinalPass_RenderThread(FRDGBuilder& GraphBuilder, const FDisplayClusterViewport_Context& InViewportContext,
-	FRHITexture2D* InputTextureRHI, const FIntRect& InputRect, FRHITexture2D* OutputTextureRHI, const FIntRect& OutputRect) const
+	FRHITexture* InputTextureRHI, const FIntRect& InputRect, FRHITexture* OutputTextureRHI, const FIntRect& OutputRect) const
 {
 	FRDGTextureRef InputTexture = GraphBuilder.RegisterExternalTexture(CreateRenderTarget(InputTextureRHI, TEXT("FDisplayClusterDisplayDeviceOCIOInputTexture")));
 	FRDGTextureRef OutputTexture = GraphBuilder.RegisterExternalTexture(CreateRenderTarget(OutputTextureRHI, TEXT("FDisplayClusterDisplayDeviceOCIORenderTargetTexture")));
diff --git a/Engine/Plugins/Runtime/nDisplay/Source/DisplayCluster/Private/Render/DisplayDevice/Proxy/DisplayClusterDisplayDeviceProxy_OpenColorIO.h b/Engine/Plugins/Runtime/nDisplay/Source/DisplayCluster/Private/Render/DisplayDevice/Proxy/DisplayClusterDisplayDeviceProxy_OpenColorIO.h
--- a/Engine/Plugins/Runtime/nDisplay/Source/DisplayCluster/Private/Render/DisplayDevice/Proxy/DisplayClusterDisplayDeviceProxy_OpenColorIO.h
+++ b/Engine/Plugins/Runtime/nDisplay/Source/DisplayCluster/Private/Render/DisplayDevice/Proxy/DisplayClusterDisplayDeviceProxy_OpenColorIO.h
@@ -21,7 +21,7 @@ public:
 public:
 	virtual bool HasFinalPass_RenderThread() const override;
 	virtual bool AddFinalPass_RenderThread(FRDGBuilder& GraphBuilder, const FDisplayClusterViewport_Context& InViewportContext,
-		FRHITexture2D* InputTextureRHI, const FIntRect& InputRect, FRHITexture2D* OutputTextureRHI, const FIntRect& OutputRect) const override;
+		FRHITexture* InputTextureRHI, const FIntRect& InputRect, FRHITexture* OutputTextureRHI, const FIntRect& OutputRect) const override;
 
 public:
 	const FString OCIOPassId;
diff --git a/Engine/Plugins/Runtime/nDisplay/Source/DisplayCluster/Private/Render/Viewport/DisplayClusterViewportManagerProxy.cpp b/Engine/Plugins/Runtime/nDisplay/Source/DisplayCluster/Private/Render/Viewport/DisplayClusterViewportManagerProxy.cpp
--- a/Engine/Plugins/Runtime/nDisplay/Source/DisplayCluster/Private/Render/Viewport/DisplayClusterViewportManagerProxy.cpp
+++ b/Engine/Plugins/Runtime/nDisplay/Source/DisplayCluster/Private/Render/Viewport/DisplayClusterViewportManagerProxy.cpp
@@ -254,7 +254,7 @@ void FDisplayClusterViewportManagerProxy::ImplRenderFrame_GameThread(FViewport*
 
 		if (OutputViewport)
 		{
-			if (FRHITexture2D* FrameOutputRTT = OutputViewport->GetRenderTargetTexture())
+			if (FRHITexture* FrameOutputRTT = OutputViewport->GetRenderTargetTexture())
 			{
 				// For quadbuf stereo copy only left eye, right copy from OutputFrameTarget
 				//@todo Copy QuadBuf_LeftEye/(mono,sbs,tp) to separate rtt, before UI and debug rendering
@@ -296,12 +296,12 @@ void FDisplayClusterViewportManagerProxy::UpdateDeferredResources_RenderThread(F
 
 void FDisplayClusterViewportManagerProxy::ImplClearFrameTargets_RenderThread(FRHICommandListImmediate& RHICmdList) const
 {
-	TArray<FRHITexture2D*> FrameResources;
-	TArray<FRHITexture2D*> AdditionalFrameResources;
+	TArray<FRHITexture*> FrameResources;
+	TArray<FRHITexture*> AdditionalFrameResources;
 	TArray<FIntPoint> TargetOffset;
 	if (GetFrameTargets_RenderThread(FrameResources, TargetOffset, &AdditionalFrameResources))
 	{
-		for (FRHITexture2D* FrameResourceIt : FrameResources)
+		for (FRHITexture* FrameResourceIt : FrameResources)
 		{
 			FDisplayClusterViewportProxy::FillTextureWithColor_RenderThread(RHICmdList, FrameResourceIt, FLinearColor::Black);
 		}
@@ -481,7 +481,7 @@ void FDisplayClusterViewportManagerProxy::DoCrossGPUTransfers_RenderThread(FRHIC
 			{
 				if (FRenderTarget* RenderTarget = ViewportRenderTargetResource.IsValid() ? ViewportRenderTargetResource->GetViewportResourceRenderTarget() : nullptr)
 				{
-					if (FRHITexture2D* TextureRHI = ViewportRenderTargetResource->GetViewportResourceRHI_RenderThread())
+					if (FRHITexture* TextureRHI = ViewportRenderTargetResource->GetViewportResourceRHI_RenderThread())
 					{
 						const FRHIGPUMask RenderTargetGPUMask = RenderTarget->GetGPUMask(RHICmdList);
 
@@ -511,7 +511,7 @@ void FDisplayClusterViewportManagerProxy::DoCrossGPUTransfers_RenderThread(FRHIC
 #endif // WITH_MGPU
 }
 
-bool FDisplayClusterViewportManagerProxy::GetFrameTargets_RenderThread(TArray<FRHITexture2D*>& OutFrameResources, TArray<FIntPoint>& OutTargetOffsets, TArray<FRHITexture2D*>* OutAdditionalFrameResources) const
+bool FDisplayClusterViewportManagerProxy::GetFrameTargets_RenderThread(TArray<FRHITexture*>& OutFrameResources, TArray<FIntPoint>& OutTargetOffsets, TArray<FRHITexture*>* OutAdditionalFrameResources) const
 {
 	check(IsInRenderingThread());
 
@@ -526,14 +526,14 @@ bool FDisplayClusterViewportManagerProxy::GetFrameTargets_RenderThread(TArray<FR
 
 			for (int32 FrameIt = 0; FrameIt < Frames.Num(); FrameIt++)
 			{
-				if (FRHITexture2D* FrameTexture = Frames[FrameIt].IsValid() ? Frames[FrameIt]->GetViewportResourceRHI_RenderThread() : nullptr)
+				if (FRHITexture* FrameTexture = Frames[FrameIt].IsValid() ? Frames[FrameIt]->GetViewportResourceRHI_RenderThread() : nullptr)
 				{
 					OutFrameResources.Add(FrameTexture);
 					OutTargetOffsets.Add(Frames[FrameIt]->GetBackbufferFrameOffset());
 
 					if (OutAdditionalFrameResources && AdditionalFrames.IsValidIndex(FrameIt))
 					{
-						if (FRHITexture2D* AdditionalFrameTexture = AdditionalFrames[FrameIt].IsValid() ? AdditionalFrames[FrameIt]->GetViewportResourceRHI_RenderThread() : nullptr)
+						if (FRHITexture* AdditionalFrameTexture = AdditionalFrames[FrameIt].IsValid() ? AdditionalFrames[FrameIt]->GetViewportResourceRHI_RenderThread() : nullptr)
 						{
 							OutAdditionalFrameResources->Add(AdditionalFrameTexture);
 						}
@@ -562,18 +562,18 @@ bool FDisplayClusterViewportManagerProxy::GetFrameTargets_RenderThread(TArray<FR
 	return false;
 }
 
-bool FDisplayClusterViewportManagerProxy::ResolveFrameTargetToBackBuffer_RenderThread(FRHICommandListImmediate& RHICmdList, const uint32 InContextNum, const int32 DestArrayIndex, FRHITexture2D* DestTexture, FVector2D WindowSize) const
+bool FDisplayClusterViewportManagerProxy::ResolveFrameTargetToBackBuffer_RenderThread(FRHICommandListImmediate& RHICmdList, const uint32 InContextNum, const int32 DestArrayIndex, FRHITexture* DestTexture, FVector2D WindowSize) const
 {
 	check(IsInRenderingThread());
 
-	TArray<FRHITexture2D*>   FrameResources;
+	TArray<FRHITexture*>   FrameResources;
 	TArray<FIntPoint>        TargetOffsets;
 	if (GetFrameTargets_RenderThread(FrameResources, TargetOffsets))
 	{
 		// Use internal frame textures as source
 		int32 ContextNum = InContextNum;
 
-		FRHITexture2D* FrameTexture = FrameResources[ContextNum];
+		FRHITexture* FrameTexture = FrameResources[ContextNum];
 		FIntPoint DstOffset = TargetOffsets[ContextNum];
 
 		if (FrameTexture)
diff --git a/Engine/Plugins/Runtime/nDisplay/Source/DisplayCluster/Private/Render/Viewport/DisplayClusterViewportManagerProxy.h b/Engine/Plugins/Runtime/nDisplay/Source/DisplayCluster/Private/Render/Viewport/DisplayClusterViewportManagerProxy.h
--- a/Engine/Plugins/Runtime/nDisplay/Source/DisplayCluster/Private/Render/Viewport/DisplayClusterViewportManagerProxy.h
+++ b/Engine/Plugins/Runtime/nDisplay/Source/DisplayCluster/Private/Render/Viewport/DisplayClusterViewportManagerProxy.h
@@ -64,8 +64,8 @@ public:
 		return TArrayView<TSharedPtr<IDisplayClusterViewportProxy, ESPMode::ThreadSafe>>((TSharedPtr<IDisplayClusterViewportProxy, ESPMode::ThreadSafe>*)(CurrentRenderFrameViewportProxies.GetData()), CurrentRenderFrameViewportProxies.Num());
 	}
 
-	virtual bool GetFrameTargets_RenderThread(TArray<FRHITexture2D*>& OutFrameResources, TArray<FIntPoint>& OutTargetOffsets, TArray<FRHITexture2D*>* OutAdditionalFrameResources = nullptr) const override;
-	virtual bool ResolveFrameTargetToBackBuffer_RenderThread(FRHICommandListImmediate& RHICmdList, const uint32 InContextNum, const int32 DestArrayIndex, FRHITexture2D* DestTexture, FVector2D WindowSize) const override;
+	virtual bool GetFrameTargets_RenderThread(TArray<FRHITexture*>& OutFrameResources, TArray<FIntPoint>& OutTargetOffsets, TArray<FRHITexture*>* OutAdditionalFrameResources = nullptr) const override;
+	virtual bool ResolveFrameTargetToBackBuffer_RenderThread(FRHICommandListImmediate& RHICmdList, const uint32 InContextNum, const int32 DestArrayIndex, FRHITexture* DestTexture, FVector2D WindowSize) const override;
 	//~~IDisplayClusterViewportManagerProxy
 
 
diff --git a/Engine/Plugins/Runtime/nDisplay/Source/DisplayCluster/Private/Render/Viewport/DisplayClusterViewportProxy.cpp b/Engine/Plugins/Runtime/nDisplay/Source/DisplayCluster/Private/Render/Viewport/DisplayClusterViewportProxy.cpp
--- a/Engine/Plugins/Runtime/nDisplay/Source/DisplayCluster/Private/Render/Viewport/DisplayClusterViewportProxy.cpp
+++ b/Engine/Plugins/Runtime/nDisplay/Source/DisplayCluster/Private/Render/Viewport/DisplayClusterViewportProxy.cpp
@@ -106,7 +106,7 @@ void FDisplayClusterViewportProxy::UpdateViewportProxyData_RenderThread(const FD
 
 
 //  Return viewport scene proxy resources by type
-bool FDisplayClusterViewportProxy::GetResources_RenderThread(const EDisplayClusterViewportResourceType InExtResourceType, TArray<FRHITexture2D*>& OutResources) const
+bool FDisplayClusterViewportProxy::GetResources_RenderThread(const EDisplayClusterViewportResourceType InExtResourceType, TArray<FRHITexture*>& OutResources) const
 {
 	return ImplGetResources_RenderThread(InExtResourceType, OutResources, false);
 }
@@ -190,8 +190,8 @@ void FDisplayClusterViewportProxy::ImplViewportRemap_RenderThread(FRHICommandLis
 				const TSharedPtr<FDisplayClusterViewportResource, ESPMode::ThreadSafe>& Src = Resources[EDisplayClusterViewportResource::AdditionalFrameTargetableResources][ContextIt];
 				const TSharedPtr<FDisplayClusterViewportResource, ESPMode::ThreadSafe>& Dst = Resources[EDisplayClusterViewportResource::OutputFrameTargetableResources][ContextIt];
 
-				FRHITexture2D* Input = Src.IsValid() ? Src->GetViewportResourceRHI_RenderThread() : nullptr;
-				FRHITexture2D* Output = Dst.IsValid() ? Dst->GetViewportResourceRHI_RenderThread() : nullptr;
+				FRHITexture* Input = Src.IsValid() ? Src->GetViewportResourceRHI_RenderThread() : nullptr;
+				FRHITexture* Output = Dst.IsValid() ? Dst->GetViewportResourceRHI_RenderThread() : nullptr;
 
 				if (Input && Output)
 				{
@@ -202,7 +202,7 @@ void FDisplayClusterViewportProxy::ImplViewportRemap_RenderThread(FRHICommandLis
 	}
 }
 
-bool FDisplayClusterViewportProxy::GetResourcesWithRects_RenderThread(const EDisplayClusterViewportResourceType InExtResourceType, TArray<FRHITexture2D*>& OutResources, TArray<FIntRect>& OutResourceRects) const
+bool FDisplayClusterViewportProxy::GetResourcesWithRects_RenderThread(const EDisplayClusterViewportResourceType InExtResourceType, TArray<FRHITexture*>& OutResources, TArray<FIntRect>& OutResourceRects) const
 {
 	return ImplGetResourcesWithRects_RenderThread(InExtResourceType, OutResources, OutResourceRects, false);
 }
@@ -218,7 +218,7 @@ bool FDisplayClusterViewportProxy::ApplyOCIO_RenderThread(FRHICommandListImmedia
 		? EDisplayClusterViewportResourceType::InputShaderResource
 		: EDisplayClusterViewportResourceType::AdditionalTargetableResource;
 
-	TArray<FRHITexture2D*> Input, Output;
+	TArray<FRHITexture*> Input, Output;
 	TArray<FIntRect> InputRects, OutputRects;
 	if (!InSrcViewportProxy.GetResourcesWithRects_RenderThread(InSrcResourceType, Input, InputRects)
 		|| !GetResourcesWithRects_RenderThread(DestResourceType, Output, OutputRects)
@@ -323,8 +323,8 @@ void FDisplayClusterViewportProxy::UpdateDeferredResources(FRHICommandListImmedi
 	// Pass 1: Generate blur postprocess effect for render target texture rect for all contexts
 	if (PostRenderSettings.PostprocessBlur.IsEnabled())
 	{
-		TArray<FRHITexture2D*> InShaderResources;
-		TArray<FRHITexture2D*> OutTargetableResources;
+		TArray<FRHITexture*> InShaderResources;
+		TArray<FRHITexture*> OutTargetableResources;
 		if (GetResources_RenderThread(EDisplayClusterViewportResourceType::InputShaderResource, InShaderResources) && GetResources_RenderThread(EDisplayClusterViewportResourceType::AdditionalTargetableResource, OutTargetableResources))
 		{
 			// Render postprocess blur:
@@ -341,14 +341,14 @@ void FDisplayClusterViewportProxy::UpdateDeferredResources(FRHICommandListImmedi
 	// Pass 2: Create mips texture and generate mips from render target rect for all contexts
 	if (PostRenderSettings.GenerateMips.IsEnabled())
 	{
-		TArray<FRHITexture2D*> InOutMipsResources;
+		TArray<FRHITexture*> InOutMipsResources;
 		if (GetResources_RenderThread(EDisplayClusterViewportResourceType::MipsShaderResource, InOutMipsResources))
 		{
 			// Copy input image to layer0 on mips texture
 			ResolveResources_RenderThread(RHICmdList, EDisplayClusterViewportResourceType::InputShaderResource, EDisplayClusterViewportResourceType::MipsShaderResource);
 
 			// Generate mips
-			for (FRHITexture2D*& ResourceIt : InOutMipsResources)
+			for (FRHITexture*& ResourceIt : InOutMipsResources)
 			{
 				GetShadersAPI().GenerateMips(RHICmdList, ResourceIt, PostRenderSettings.GenerateMips);
 			}
@@ -472,7 +472,7 @@ void FDisplayClusterViewportProxy::OnPostRenderViewFamily_RenderThread(FRDGBuild
 			EFXAAQuality FXAAQuality = EFXAAQuality::Q0;
 			if (ShouldApplyFXAA_RenderThread(FXAAQuality) && Resources[EDisplayClusterViewportResource::InputShaderResources].IsValidIndex(InContextNum))
 			{
-				if (FRHITexture2D* InputTextureRHI = Resources[EDisplayClusterViewportResource::InputShaderResources][InContextNum] ? Resources[EDisplayClusterViewportResource::InputShaderResources][InContextNum]->GetViewportResourceRHI_RenderThread() : nullptr)
+				if (FRHITexture* InputTextureRHI = Resources[EDisplayClusterViewportResource::InputShaderResources][InContextNum] ? Resources[EDisplayClusterViewportResource::InputShaderResources][InContextNum]->GetViewportResourceRHI_RenderThread() : nullptr)
 				{
 					// Apply FXAA for RGB only
 					// Note: Add AA for alpha channel
diff --git a/Engine/Plugins/Runtime/nDisplay/Source/DisplayCluster/Private/Render/Viewport/DisplayClusterViewportProxy.h b/Engine/Plugins/Runtime/nDisplay/Source/DisplayCluster/Private/Render/Viewport/DisplayClusterViewportProxy.h
--- a/Engine/Plugins/Runtime/nDisplay/Source/DisplayCluster/Private/Render/Viewport/DisplayClusterViewportProxy.h
+++ b/Engine/Plugins/Runtime/nDisplay/Source/DisplayCluster/Private/Render/Viewport/DisplayClusterViewportProxy.h
@@ -136,8 +136,8 @@ public:
 	}
 
 	//  Return viewport scene proxy resources by type
-	virtual bool GetResources_RenderThread(const EDisplayClusterViewportResourceType InResourceType, TArray<FRHITexture2D*>& OutResources) const override;
-	virtual bool GetResourcesWithRects_RenderThread(const EDisplayClusterViewportResourceType InResourceType, TArray<FRHITexture2D*>& OutResources, TArray<FIntRect>& OutRects) const override;
+	virtual bool GetResources_RenderThread(const EDisplayClusterViewportResourceType InResourceType, TArray<FRHITexture*>& OutResources) const override;
+	virtual bool GetResourcesWithRects_RenderThread(const EDisplayClusterViewportResourceType InResourceType, TArray<FRHITexture*>& OutResources, TArray<FIntRect>& OutRects) const override;
 
 	// Resolve resource contexts
 	virtual bool ResolveResources_RenderThread(FRHICommandListImmediate& RHICmdList, const EDisplayClusterViewportResourceType InputResourceType, const EDisplayClusterViewportResourceType OutputResourceType, const int32 InContextNum = INDEX_NONE) const override;
@@ -152,7 +152,7 @@ public:
 	* @param InRenderTargetTexture - texture,  must be RTT.
 	* @param InColor - the texture will be filled with this color.
 	*/
-	static void FillTextureWithColor_RenderThread(FRHICommandListImmediate& RHICmdList, FRHITexture2D* InRenderTargetTexture, const FLinearColor& InColor);
+	static void FillTextureWithColor_RenderThread(FRHICommandListImmediate& RHICmdList, FRHITexture* InRenderTargetTexture, const FLinearColor& InColor);
 
 	/** Get valid resource type
 	 * 
@@ -299,8 +299,8 @@ public:
 	uint8 GetPriority_RenderThread() const;
 
 private:
-	bool ImplGetResourcesWithRects_RenderThread(const EDisplayClusterViewportResourceType InResourceType, TArray<FRHITexture2D*>& OutResources, TArray<FIntRect>& OutResourceRects, const int32 InRecursionDepth) const;
-	bool ImplGetResources_RenderThread(const EDisplayClusterViewportResourceType InResourceType, TArray<FRHITexture2D*>& OutResources, const int32 InRecursionDepth) const;
+	bool ImplGetResourcesWithRects_RenderThread(const EDisplayClusterViewportResourceType InResourceType, TArray<FRHITexture*>& OutResources, TArray<FIntRect>& OutResourceRects, const int32 InRecursionDepth) const;
+	bool ImplGetResources_RenderThread(const EDisplayClusterViewportResourceType InResourceType, TArray<FRHITexture*>& OutResources, const int32 InRecursionDepth) const;
 
 	void ImplViewportRemap_RenderThread(FRHICommandListImmediate& RHICmdList) const;
 
diff --git a/Engine/Plugins/Runtime/nDisplay/Source/DisplayCluster/Private/Render/Viewport/DisplayClusterViewportProxy_Resources.cpp b/Engine/Plugins/Runtime/nDisplay/Source/DisplayCluster/Private/Render/Viewport/DisplayClusterViewportProxy_Resources.cpp
--- a/Engine/Plugins/Runtime/nDisplay/Source/DisplayCluster/Private/Render/Viewport/DisplayClusterViewportProxy_Resources.cpp
+++ b/Engine/Plugins/Runtime/nDisplay/Source/DisplayCluster/Private/Render/Viewport/DisplayClusterViewportProxy_Resources.cpp
@@ -140,7 +140,7 @@ namespace UE::DisplayCluster::ViewportProxy
 		RHICmdList.Transition(FRHITransitionInfo(DstTexture, ERHIAccess::Unknown, ERHIAccess::SRVMask));
 	}
 
-	static void ImplResolveResource(FRHICommandListImmediate& RHICmdList, FRHITexture2D* InputResource, const FIntRect& InputRect, FRHITexture2D* OutputResource, const FIntRect& OutputRect, const bool bOutputIsMipsResource, const bool bOutputIsPreviewResource)
+	static void ImplResolveResource(FRHICommandListImmediate& RHICmdList, FRHITexture* InputResource, const FIntRect& InputRect, FRHITexture* OutputResource, const FIntRect& OutputRect, const bool bOutputIsMipsResource, const bool bOutputIsPreviewResource)
 	{
 		check(InputResource);
 		check(OutputResource);
@@ -169,7 +169,7 @@ namespace UE::DisplayCluster::ViewportProxy
 
 	struct FViewportResourceResolverData
 	{
-		FViewportResourceResolverData(const FDisplayClusterViewport_Context& InViewportContext, FRHITexture2D* InInputResource, const FIntRect& InInputRect, FRHITexture2D* InOutputResource, const FIntRect& InOutputRect, const bool InbOutputIsMipsResource, const bool InbOutputIsPreviewResource)
+		FViewportResourceResolverData(const FDisplayClusterViewport_Context& InViewportContext, FRHITexture* InInputResource, const FIntRect& InInputRect, FRHITexture* InOutputResource, const FIntRect& InOutputRect, const bool InbOutputIsMipsResource, const bool InbOutputIsPreviewResource)
 			: InputResource(InInputResource)
 			, OutputResource(InOutputResource)
 			, InputRect(InInputRect)
@@ -190,8 +190,8 @@ namespace UE::DisplayCluster::ViewportProxy
 		}
 
 	private:
-		FRHITexture2D* InputResource;
-		FRHITexture2D* OutputResource;
+		FRHITexture* InputResource;
+		FRHITexture* OutputResource;
 
 		const FIntRect InputRect;
 		const FIntRect OutputRect;
@@ -203,7 +203,7 @@ namespace UE::DisplayCluster::ViewportProxy
 	};
 };
 
-void FDisplayClusterViewportProxy::FillTextureWithColor_RenderThread(FRHICommandListImmediate& RHICmdList, FRHITexture2D* InRenderTargetTexture, const FLinearColor& InColor)
+void FDisplayClusterViewportProxy::FillTextureWithColor_RenderThread(FRHICommandListImmediate& RHICmdList, FRHITexture* InRenderTargetTexture, const FLinearColor& InColor)
 {
 	if (InRenderTargetTexture)
 	{
@@ -220,7 +220,7 @@ void FDisplayClusterViewportProxy::FillTextureWithColor_RenderThread(FRHICommand
 	}
 }
 
-bool FDisplayClusterViewportProxy::ImplGetResources_RenderThread(const EDisplayClusterViewportResourceType InExtResourceType, TArray<FRHITexture2D*>& OutResources, const int32 InRecursionDepth) const
+bool FDisplayClusterViewportProxy::ImplGetResources_RenderThread(const EDisplayClusterViewportResourceType InExtResourceType, TArray<FRHITexture*>& OutResources, const int32 InRecursionDepth) const
 {
 	using namespace UE::DisplayCluster::ViewportProxy;
 	check(IsInRenderingThread());
@@ -255,7 +255,7 @@ bool FDisplayClusterViewportProxy::ImplGetResources_RenderThread(const EDisplayC
 				bResult = true;
 
 				// Support texture replace:
-				if (FRHITexture2D* ReplaceTextureRHI = PostRenderSettings.Replace.TextureRHI->GetTexture2D())
+				if (FRHITexture* ReplaceTextureRHI = PostRenderSettings.Replace.TextureRHI->GetTexture2D())
 				{
 					for (int32 ContextIndex = 0; ContextIndex < Contexts.Num(); ContextIndex++)
 					{
@@ -326,7 +326,7 @@ bool FDisplayClusterViewportProxy::ImplGetResources_RenderThread(const EDisplayC
 	return false;
 }
 
-bool FDisplayClusterViewportProxy::ImplGetResourcesWithRects_RenderThread(const EDisplayClusterViewportResourceType InExtResourceType, TArray<FRHITexture2D*>& OutResources, TArray<FIntRect>& OutResourceRects, const int32 InRecursionDepth) const
+bool FDisplayClusterViewportProxy::ImplGetResourcesWithRects_RenderThread(const EDisplayClusterViewportResourceType InExtResourceType, TArray<FRHITexture*>& OutResources, TArray<FIntRect>& OutResourceRects, const int32 InRecursionDepth) const
 {
 	using namespace UE::DisplayCluster::ViewportProxy;
 	check(IsInRenderingThread());
@@ -406,7 +406,7 @@ bool FDisplayClusterViewportProxy::ImplResolveResources_RenderThread(FRHICommand
 
 	TArray<FViewportResourceResolverData> ResourceResolverData;
 
-	TArray<FRHITexture2D*> SrcResources, DestResources;
+	TArray<FRHITexture*> SrcResources, DestResources;
 	TArray<FIntRect> SrcResourcesRect, DestResourcesRect;
 	if (SourceProxy->GetResourcesWithRects_RenderThread(InExtResourceType, SrcResources, SrcResourcesRect) && GetResourcesWithRects_RenderThread(OutExtResourceType, DestResources, DestResourcesRect))
 	{
@@ -470,7 +470,7 @@ bool FDisplayClusterViewportProxy::ImplResolveResources_RenderThread(FRHICommand
 
 bool FDisplayClusterViewportProxy::CopyResource_RenderThread(FRDGBuilder& GraphBuilder, const EDisplayClusterTextureCopyMode InCopyMode, const int32 InContextNum, const EDisplayClusterViewportResourceType InExtSrcResourceType, const EDisplayClusterViewportResourceType InExtDestResourceType)
 {
-	TArray<FRHITexture2D*> SrcResources;
+	TArray<FRHITexture*> SrcResources;
 	TArray<FIntRect> SrcResourceRects;
 	if (GetResourcesWithRects_RenderThread(InExtSrcResourceType, SrcResources, SrcResourceRects))
 	{
@@ -497,7 +497,7 @@ bool FDisplayClusterViewportProxy::CopyResource_RenderThread(FRDGBuilder& GraphB
 {
 	using namespace UE::DisplayCluster::ViewportProxy;
 
-	TArray<FRHITexture2D*> DestResources;
+	TArray<FRHITexture*> DestResources;
 	TArray<FIntRect> DestResourceRects;
 	if (HasBeenProduced(InSrcTextureRef) && GetResourcesWithRects_RenderThread(InExtDestResourceType, DestResources, DestResourceRects))
 	{
@@ -531,7 +531,7 @@ bool FDisplayClusterViewportProxy::CopyResource_RenderThread(FRDGBuilder& GraphB
 {
 	using namespace UE::DisplayCluster::ViewportProxy;
 
-	TArray<FRHITexture2D*> SrcResources;
+	TArray<FRHITexture*> SrcResources;
 	TArray<FIntRect> SrcResourceRects;
 	if (HasBeenProduced(InDestTextureRef) && GetResourcesWithRects_RenderThread(InExtSrcResourceType, SrcResources, SrcResourceRects))
 	{
@@ -570,7 +570,7 @@ void FDisplayClusterViewportProxy::ImplResolveTileResource_RenderThread(FRHIComm
 	const TArray<FDisplayClusterViewport_Context>& DestContexts = InDestViewportProxy->GetContexts_RenderThread();
 	const TArray<FDisplayClusterViewport_Context>& SrcContexts = Contexts;
 
-	TArray<FRHITexture2D*> SrcResources, DestResources;
+	TArray<FRHITexture*> SrcResources, DestResources;
 	if (GetResources_RenderThread(EDisplayClusterViewportResourceType::InternalRenderTargetResource, SrcResources)
 	&& InDestViewportProxy->GetResources_RenderThread(EDisplayClusterViewportResourceType::InternalRenderTargetResource, DestResources)
 	&& SrcResources.Num() == DestResources.Num()
@@ -614,10 +614,10 @@ void FDisplayClusterViewportProxy::CleanupResources_RenderThread(FRHICommandList
 	// Since the RTT is reused through frames, in case we need to show a black border between viewport tiles, we must fill the original viewport with this colour.
 	if (GDisplayClusterRenderTileBorder > 0 && RenderSettings.TileSettings.GetType() == EDisplayClusterViewportTileType::Source)
 	{
-		TArray<FRHITexture2D*> RenderTargets;
+		TArray<FRHITexture*> RenderTargets;
 		if (GetResources_RenderThread(EDisplayClusterViewportResourceType::InternalRenderTargetResource, RenderTargets))
 		{
-			for (FRHITexture2D* TextureIt : RenderTargets)
+			for (FRHITexture* TextureIt : RenderTargets)
 			{
 				// Note: It may make sense to move the CVar and border color to the StageSettings.
 				FDisplayClusterViewportProxy::FillTextureWithColor_RenderThread(RHICmdList, TextureIt, FLinearColor::Black);
diff --git a/Engine/Plugins/Runtime/nDisplay/Source/DisplayCluster/Private/Render/Viewport/DisplayClusterViewportResources.cpp b/Engine/Plugins/Runtime/nDisplay/Source/DisplayCluster/Private/Render/Viewport/DisplayClusterViewportResources.cpp
--- a/Engine/Plugins/Runtime/nDisplay/Source/DisplayCluster/Private/Render/Viewport/DisplayClusterViewportResources.cpp
+++ b/Engine/Plugins/Runtime/nDisplay/Source/DisplayCluster/Private/Render/Viewport/DisplayClusterViewportResources.cpp
@@ -19,7 +19,7 @@ void FDisplayClusterViewportResources::FreezeRendering(const EDisplayClusterView
 	}
 }
 
-bool FDisplayClusterViewportResources::GetRHIResources_RenderThread(const EDisplayClusterViewportResource InResourceType, TArray<FRHITexture2D*>& OutResources) const
+bool FDisplayClusterViewportResources::GetRHIResources_RenderThread(const EDisplayClusterViewportResource InResourceType, TArray<FRHITexture*>& OutResources) const
 {
 	OutResources.Reset();
 
@@ -32,7 +32,7 @@ bool FDisplayClusterViewportResources::GetRHIResources_RenderThread(const EDispl
 				// When resource accessed on rendering thread, update this flag
 				EnumAddFlags(ViewportResourceIt->GetResourceState(), EDisplayClusterViewportResourceState::UpdatedOnRenderingThread);
 
-				if (FRHITexture2D* RHITexture2D = ViewportResourceIt->GetViewportResourceRHI_RenderThread())
+				if (FRHITexture* RHITexture2D = ViewportResourceIt->GetViewportResourceRHI_RenderThread())
 				{
 					// Collects only valid resources.
 					OutResources.Add(RHITexture2D);
diff --git a/Engine/Plugins/Runtime/nDisplay/Source/DisplayCluster/Private/Render/Viewport/DisplayClusterViewportResources.h b/Engine/Plugins/Runtime/nDisplay/Source/DisplayCluster/Private/Render/Viewport/DisplayClusterViewportResources.h
--- a/Engine/Plugins/Runtime/nDisplay/Source/DisplayCluster/Private/Render/Viewport/DisplayClusterViewportResources.h
+++ b/Engine/Plugins/Runtime/nDisplay/Source/DisplayCluster/Private/Render/Viewport/DisplayClusterViewportResources.h
@@ -106,7 +106,7 @@ public:
 	*
 	* @return true, if all resources are valid
 	*/
-	bool GetRHIResources_RenderThread(const EDisplayClusterViewportResource InResourceType, TArray<FRHITexture2D*>& OutResources) const;
+	bool GetRHIResources_RenderThread(const EDisplayClusterViewportResource InResourceType, TArray<FRHITexture*>& OutResources) const;
 
 protected:
 	/** Release specified type of resources. */
diff --git a/Engine/Plugins/Runtime/nDisplay/Source/DisplayCluster/Private/Render/Viewport/DisplayClusterViewport_OpenColorIO.cpp b/Engine/Plugins/Runtime/nDisplay/Source/DisplayCluster/Private/Render/Viewport/DisplayClusterViewport_OpenColorIO.cpp
--- a/Engine/Plugins/Runtime/nDisplay/Source/DisplayCluster/Private/Render/Viewport/DisplayClusterViewport_OpenColorIO.cpp
+++ b/Engine/Plugins/Runtime/nDisplay/Source/DisplayCluster/Private/Render/Viewport/DisplayClusterViewport_OpenColorIO.cpp
@@ -62,7 +62,7 @@ float FDisplayClusterViewport_OpenColorIO::GetDisplayGamma(const FDisplayCluster
 }
 
 bool FDisplayClusterViewport_OpenColorIO::AddPass_RenderThread(FRDGBuilder& GraphBuilder, const FDisplayClusterViewport_Context& InViewportContext,
-	FRHITexture2D* InputTextureRHI, const FIntRect& InputRect, FRHITexture2D* OutputTextureRHI, const FIntRect& OutputRect, bool bUnpremultiply, bool bInvertAlpha) const
+	FRHITexture* InputTextureRHI, const FIntRect& InputRect, FRHITexture* OutputTextureRHI, const FIntRect& OutputRect, bool bUnpremultiply, bool bInvertAlpha) const
 {
 	FRDGTextureRef InputTexture = GraphBuilder.RegisterExternalTexture(CreateRenderTarget(InputTextureRHI, TEXT("nDisplayOCIOInputTexture")));
 	FRDGTextureRef OutputTexture = GraphBuilder.RegisterExternalTexture(CreateRenderTarget(OutputTextureRHI, TEXT("nDisplayOCIORenderTargetTexture")));
diff --git a/Engine/Plugins/Runtime/nDisplay/Source/DisplayCluster/Private/Render/Viewport/DisplayClusterViewport_OpenColorIO.h b/Engine/Plugins/Runtime/nDisplay/Source/DisplayCluster/Private/Render/Viewport/DisplayClusterViewport_OpenColorIO.h
--- a/Engine/Plugins/Runtime/nDisplay/Source/DisplayCluster/Private/Render/Viewport/DisplayClusterViewport_OpenColorIO.h
+++ b/Engine/Plugins/Runtime/nDisplay/Source/DisplayCluster/Private/Render/Viewport/DisplayClusterViewport_OpenColorIO.h
@@ -53,7 +53,7 @@ public:
 	 * @return - true if success.
 	 */
 	bool AddPass_RenderThread(FRDGBuilder& GraphBuilder, const FDisplayClusterViewport_Context& InViewportContext,
-		FRHITexture2D* InputTextureRHI, const FIntRect& InputRect, FRHITexture2D* OutputTextureRHI, const FIntRect& OutputRect, bool bUnpremultiply, bool InvertedAlpha) const;
+		FRHITexture* InputTextureRHI, const FIntRect& InputRect, FRHITexture* OutputTextureRHI, const FIntRect& OutputRect, bool bUnpremultiply, bool InvertedAlpha) const;
 
 	/* This is a copy of FOpenColorIODisplayExtension::PostProcessPassAfterTonemap_RenderThread() */
 	FScreenPassTexture PostProcessPassAfterTonemap_RenderThread(FRDGBuilder& GraphBuilder, const FDisplayClusterViewport_Context& InViewportContext, const FSceneView& View, const FPostProcessMaterialInputs& InOutInputs);
diff --git a/Engine/Plugins/Runtime/nDisplay/Source/DisplayCluster/Private/Render/Viewport/Postprocess/DisplayClusterViewportPostProcessManager.cpp b/Engine/Plugins/Runtime/nDisplay/Source/DisplayCluster/Private/Render/Viewport/Postprocess/DisplayClusterViewportPostProcessManager.cpp
--- a/Engine/Plugins/Runtime/nDisplay/Source/DisplayCluster/Private/Render/Viewport/Postprocess/DisplayClusterViewportPostProcessManager.cpp
+++ b/Engine/Plugins/Runtime/nDisplay/Source/DisplayCluster/Private/Render/Viewport/Postprocess/DisplayClusterViewportPostProcessManager.cpp
@@ -536,8 +536,8 @@ void FDisplayClusterViewportPostProcessManager::ImplPerformPostProcessFrameAfter
 
 	if (bEnabled != 0 && InViewportManagerProxy)
 	{
-		TArray<FRHITexture2D*> FrameResources;
-		TArray<FRHITexture2D*> AdditionalFrameResources;
+		TArray<FRHITexture*> FrameResources;
+		TArray<FRHITexture*> AdditionalFrameResources;
 		TArray<FIntPoint> TargetOffset;
 		if (InViewportManagerProxy->GetFrameTargets_RenderThread(FrameResources, TargetOffset, &AdditionalFrameResources))
 		{
@@ -547,7 +547,7 @@ void FDisplayClusterViewportPostProcessManager::ImplPerformPostProcessFrameAfter
 				{
 					UE_LOG(LogDisplayClusterRender, VeryVerbose, TEXT("Postprocess FRAME after WarpBlend"));
 
-					TArray<FRHITexture2D*>* AdditionalResources = (AdditionalFrameResources.Num() > 0 && It->ShouldUseAdditionalFrameTargetableResource())? &AdditionalFrameResources: nullptr;
+					TArray<FRHITexture*>* AdditionalResources = (AdditionalFrameResources.Num() > 0 && It->ShouldUseAdditionalFrameTargetableResource())? &AdditionalFrameResources: nullptr;
 					It->PerformPostProcessFrameAfterWarpBlend_RenderThread(RHICmdList, &FrameResources, AdditionalResources);
 				}
 			}
@@ -555,7 +555,7 @@ void FDisplayClusterViewportPostProcessManager::ImplPerformPostProcessFrameAfter
 			// Apply OutputRemap after all postprocess
 			if (OutputRemap.IsValid() && OutputRemap->IsEnabled())
 			{
-				TArray<FRHITexture2D*>* AdditionalResources = (AdditionalFrameResources.Num() > 0) ? &AdditionalFrameResources : nullptr;
+				TArray<FRHITexture*>* AdditionalResources = (AdditionalFrameResources.Num() > 0) ? &AdditionalFrameResources : nullptr;
 				OutputRemap->PerformPostProcessFrame_RenderThread(RHICmdList, &FrameResources, AdditionalResources);
 			}
 		}
diff --git a/Engine/Plugins/Runtime/nDisplay/Source/DisplayCluster/Private/Render/Viewport/Postprocess/DisplayClusterViewportPostProcessOutputRemap.cpp b/Engine/Plugins/Runtime/nDisplay/Source/DisplayCluster/Private/Render/Viewport/Postprocess/DisplayClusterViewportPostProcessOutputRemap.cpp
--- a/Engine/Plugins/Runtime/nDisplay/Source/DisplayCluster/Private/Render/Viewport/Postprocess/DisplayClusterViewportPostProcessOutputRemap.cpp
+++ b/Engine/Plugins/Runtime/nDisplay/Source/DisplayCluster/Private/Render/Viewport/Postprocess/DisplayClusterViewportPostProcessOutputRemap.cpp
@@ -222,7 +222,7 @@ bool FDisplayClusterViewportPostProcessOutputRemap::MarkProceduralMeshComponentG
 	return false;
 }
 
-void FDisplayClusterViewportPostProcessOutputRemap::PerformPostProcessFrame_RenderThread(FRHICommandListImmediate& RHICmdList, const TArray<FRHITexture2D*>* InFrameTargets, const TArray<FRHITexture2D*>* InAdditionalFrameTargets) const
+void FDisplayClusterViewportPostProcessOutputRemap::PerformPostProcessFrame_RenderThread(FRHICommandListImmediate& RHICmdList, const TArray<FRHITexture*>* InFrameTargets, const TArray<FRHITexture*>* InAdditionalFrameTargets) const
 {
 	check(IsInRenderingThread());
 
@@ -233,8 +233,8 @@ void FDisplayClusterViewportPostProcessOutputRemap::PerformPostProcessFrame_Rend
 		{
 			for (int32 Index = 0; Index < InFrameTargets->Num(); Index++)
 			{
-				FRHITexture2D* InOutTexture = (*InFrameTargets)[Index];
-				FRHITexture2D* TempTargetableTexture = (*InAdditionalFrameTargets)[Index];
+				FRHITexture* InOutTexture = (*InFrameTargets)[Index];
+				FRHITexture* TempTargetableTexture = (*InAdditionalFrameTargets)[Index];
 
 				static IDisplayClusterShaders& ShadersAPI = IDisplayClusterShaders::Get();
 				if (ShadersAPI.RenderPostprocess_OutputRemap(RHICmdList, InOutTexture, TempTargetableTexture, *MeshProxy))
diff --git a/Engine/Plugins/Runtime/nDisplay/Source/DisplayCluster/Private/Render/Viewport/Postprocess/DisplayClusterViewportPostProcessOutputRemap.h b/Engine/Plugins/Runtime/nDisplay/Source/DisplayCluster/Private/Render/Viewport/Postprocess/DisplayClusterViewportPostProcessOutputRemap.h
--- a/Engine/Plugins/Runtime/nDisplay/Source/DisplayCluster/Private/Render/Viewport/Postprocess/DisplayClusterViewportPostProcessOutputRemap.h
+++ b/Engine/Plugins/Runtime/nDisplay/Source/DisplayCluster/Private/Render/Viewport/Postprocess/DisplayClusterViewportPostProcessOutputRemap.h
@@ -34,7 +34,7 @@ public:
 	}
 
 public:
-	void PerformPostProcessFrame_RenderThread(FRHICommandListImmediate& RHICmdList, const TArray<FRHITexture2D*>* InFrameTargets = nullptr, const TArray<FRHITexture2D*>* InAdditionalFrameTargets = nullptr) const;
+	void PerformPostProcessFrame_RenderThread(FRHICommandListImmediate& RHICmdList, const TArray<FRHITexture*>* InFrameTargets = nullptr, const TArray<FRHITexture*>* InAdditionalFrameTargets = nullptr) const;
 
 private:
 	bool ImplInitializeOutputRemap();
diff --git a/Engine/Plugins/Runtime/nDisplay/Source/DisplayCluster/Private/Render/Viewport/RenderTarget/DisplayClusterRenderTargetResource.cpp b/Engine/Plugins/Runtime/nDisplay/Source/DisplayCluster/Private/Render/Viewport/RenderTarget/DisplayClusterRenderTargetResource.cpp
--- a/Engine/Plugins/Runtime/nDisplay/Source/DisplayCluster/Private/Render/Viewport/RenderTarget/DisplayClusterRenderTargetResource.cpp
+++ b/Engine/Plugins/Runtime/nDisplay/Source/DisplayCluster/Private/Render/Viewport/RenderTarget/DisplayClusterRenderTargetResource.cpp
@@ -18,11 +18,11 @@ FDisplayClusterViewportTextureResource::FDisplayClusterViewportTextureResource(c
 
 void FDisplayClusterViewportTextureResource::InitRHI(FRHICommandListBase&)
 {
-	FTexture2DRHIRef NewTextureRHI;
+	FTextureRHIRef NewTextureRHI;
 
 	if (GetResourceSettings().GetNumMips() > 1)
 	{
-		FTexture2DRHIRef DummyTextureRHI;
+		FTextureRHIRef DummyTextureRHI;
 		ImplInitDynamicRHI_RenderTargetResource2D(NewTextureRHI, DummyTextureRHI);
 	}
 	else
@@ -39,7 +39,7 @@ void FDisplayClusterViewportTextureResource::InitRHI(FRHICommandListBase&)
 void FDisplayClusterViewportRenderTargetResource::InitRHI(FRHICommandListBase&)
 {
 	// Create RTT and shader resources
-	FTexture2DRHIRef NewTextureRHI;
+	FTextureRHIRef NewTextureRHI;
 	ImplInitDynamicRHI_RenderTargetResource2D(RenderTargetTextureRHI, NewTextureRHI);
 	TextureRHI = (FTextureRHIRef&)NewTextureRHI;
 
diff --git a/Engine/Plugins/Runtime/nDisplay/Source/DisplayCluster/Private/Render/Viewport/RenderTarget/DisplayClusterRenderTargetResource.h b/Engine/Plugins/Runtime/nDisplay/Source/DisplayCluster/Private/Render/Viewport/RenderTarget/DisplayClusterRenderTargetResource.h
--- a/Engine/Plugins/Runtime/nDisplay/Source/DisplayCluster/Private/Render/Viewport/RenderTarget/DisplayClusterRenderTargetResource.h
+++ b/Engine/Plugins/Runtime/nDisplay/Source/DisplayCluster/Private/Render/Viewport/RenderTarget/DisplayClusterRenderTargetResource.h
@@ -23,7 +23,7 @@ public:
 
 public:
 	//~ Begin FDisplayClusterViewportResource
-	virtual FRHITexture2D* GetViewportResourceRHI_RenderThread() const override
+	virtual FRHITexture* GetViewportResourceRHI_RenderThread() const override
 	{
 		check(IsInRenderingThread());
 
diff --git a/Engine/Plugins/Runtime/nDisplay/Source/DisplayCluster/Private/Render/Viewport/Resource/DisplayClusterViewportPreviewResource.cpp b/Engine/Plugins/Runtime/nDisplay/Source/DisplayCluster/Private/Render/Viewport/Resource/DisplayClusterViewportPreviewResource.cpp
--- a/Engine/Plugins/Runtime/nDisplay/Source/DisplayCluster/Private/Render/Viewport/Resource/DisplayClusterViewportPreviewResource.cpp
+++ b/Engine/Plugins/Runtime/nDisplay/Source/DisplayCluster/Private/Render/Viewport/Resource/DisplayClusterViewportPreviewResource.cpp
@@ -58,7 +58,7 @@ void FDisplayClusterViewportPreviewResource::ReleaseViewportResource()
 	EnumRemoveFlags(GetResourceState(), EDisplayClusterViewportResourceState::Initialized);
 }
 
-FRHITexture2D* FDisplayClusterViewportPreviewResource::GetViewportResourceRHI_RenderThread() const
+FRHITexture* FDisplayClusterViewportPreviewResource::GetViewportResourceRHI_RenderThread() const
 {
 	check(IsInRenderingThread());
 
diff --git a/Engine/Plugins/Runtime/nDisplay/Source/DisplayCluster/Private/Render/Viewport/Resource/DisplayClusterViewportPreviewResource.h b/Engine/Plugins/Runtime/nDisplay/Source/DisplayCluster/Private/Render/Viewport/Resource/DisplayClusterViewportPreviewResource.h
--- a/Engine/Plugins/Runtime/nDisplay/Source/DisplayCluster/Private/Render/Viewport/Resource/DisplayClusterViewportPreviewResource.h
+++ b/Engine/Plugins/Runtime/nDisplay/Source/DisplayCluster/Private/Render/Viewport/Resource/DisplayClusterViewportPreviewResource.h
@@ -28,7 +28,7 @@ public:
 		return RenderTargetTexture;
 	}
 
-	virtual FRHITexture2D* GetViewportResourceRHI_RenderThread() const override;
+	virtual FRHITexture* GetViewportResourceRHI_RenderThread() const override;
 	//~~ End FDisplayClusterViewportResource
 
 public:
diff --git a/Engine/Plugins/Runtime/nDisplay/Source/DisplayCluster/Private/Render/Viewport/Resource/DisplayClusterViewportResource.cpp b/Engine/Plugins/Runtime/nDisplay/Source/DisplayCluster/Private/Render/Viewport/Resource/DisplayClusterViewportResource.cpp
--- a/Engine/Plugins/Runtime/nDisplay/Source/DisplayCluster/Private/Render/Viewport/Resource/DisplayClusterViewportResource.cpp
+++ b/Engine/Plugins/Runtime/nDisplay/Source/DisplayCluster/Private/Render/Viewport/Resource/DisplayClusterViewportResource.cpp
@@ -15,7 +15,7 @@ FDisplayClusterViewportResource::~FDisplayClusterViewportResource()
 	check(!EnumHasAnyFlags(ResourceState, EDisplayClusterViewportResourceState::Initialized));
 }
 
-void FDisplayClusterViewportResource::ImplInitDynamicRHI_RenderTargetResource2D(FTexture2DRHIRef& OutRenderTargetTextureRHI, FTexture2DRHIRef& OutTextureRHI)
+void FDisplayClusterViewportResource::ImplInitDynamicRHI_RenderTargetResource2D(FTextureRHIRef& OutRenderTargetTextureRHI, FTextureRHIRef& OutTextureRHI)
 {
 	ETextureCreateFlags CreateFlags = TexCreate_Dynamic;
 
@@ -48,7 +48,7 @@ void FDisplayClusterViewportResource::ImplInitDynamicRHI_RenderTargetResource2D(
 	OutRenderTargetTextureRHI = OutTextureRHI = RHICreateTexture(Desc);
 }
 
-void FDisplayClusterViewportResource::ImplInitDynamicRHI_TextureResource2D(FTexture2DRHIRef& OutTextureRHI)
+void FDisplayClusterViewportResource::ImplInitDynamicRHI_TextureResource2D(FTextureRHIRef& OutTextureRHI)
 {
 	FRHITextureCreateDesc Desc =
 		FRHITextureCreateDesc::Create2D(TEXT("DisplayClusterViewportTextureResource"), ResourceSettings.GetSizeXY().X, ResourceSettings.GetSizeXY().Y, ResourceSettings.GetFormat())
diff --git a/Engine/Plugins/Runtime/nDisplay/Source/DisplayCluster/Private/Render/Viewport/Resource/DisplayClusterViewportResource.h b/Engine/Plugins/Runtime/nDisplay/Source/DisplayCluster/Private/Render/Viewport/Resource/DisplayClusterViewportResource.h
--- a/Engine/Plugins/Runtime/nDisplay/Source/DisplayCluster/Private/Render/Viewport/Resource/DisplayClusterViewportResource.h
+++ b/Engine/Plugins/Runtime/nDisplay/Source/DisplayCluster/Private/Render/Viewport/Resource/DisplayClusterViewportResource.h
@@ -72,7 +72,7 @@ public:
 	{ }
 
 	/** Get current RHI resource. */
-	virtual FRHITexture2D* GetViewportResourceRHI_RenderThread() const
+	virtual FRHITexture* GetViewportResourceRHI_RenderThread() const
 	{
 		return nullptr;
 	}
@@ -120,10 +120,10 @@ public:
 
 protected:
 	/** Helper function: create RTT texture. */
-	void ImplInitDynamicRHI_RenderTargetResource2D(FTexture2DRHIRef& OutRenderTargetTextureRHI, FTexture2DRHIRef& OutTextureRHI);
+	void ImplInitDynamicRHI_RenderTargetResource2D(FTextureRHIRef& OutRenderTargetTextureRHI, FTextureRHIRef& OutTextureRHI);
 
 	/** Helper function: create 2D texture. */
-	void ImplInitDynamicRHI_TextureResource2D(FTexture2DRHIRef& OutTextureRHI);
+	void ImplInitDynamicRHI_TextureResource2D(FTextureRHIRef& OutTextureRHI);
 
 protected:
 	// Settings used by this resource
diff --git a/Engine/Plugins/Runtime/nDisplay/Source/DisplayCluster/Private/Render/Viewport/Resource/DisplayClusterViewportResourceSettings.cpp b/Engine/Plugins/Runtime/nDisplay/Source/DisplayCluster/Private/Render/Viewport/Resource/DisplayClusterViewportResourceSettings.cpp
--- a/Engine/Plugins/Runtime/nDisplay/Source/DisplayCluster/Private/Render/Viewport/Resource/DisplayClusterViewportResourceSettings.cpp
+++ b/Engine/Plugins/Runtime/nDisplay/Source/DisplayCluster/Private/Render/Viewport/Resource/DisplayClusterViewportResourceSettings.cpp
@@ -16,7 +16,7 @@ FDisplayClusterViewportResourceSettings::FDisplayClusterViewportResourceSettings
 {
 	if (InViewport && InViewport->GetRenderTargetTexture())
 	{
-		FRHITexture2D* ViewportTexture = InViewport->GetRenderTargetTexture();
+		FRHITexture* ViewportTexture = InViewport->GetRenderTargetTexture();
 		Format = ViewportTexture->GetFormat();
 
 		if (EnumHasAnyFlags(ViewportTexture->GetFlags(), TexCreate_SRGB))
diff --git a/Engine/Plugins/Runtime/nDisplay/Source/DisplayCluster/Public/Render/DisplayDevice/IDisplayClusterDisplayDeviceProxy.h b/Engine/Plugins/Runtime/nDisplay/Source/DisplayCluster/Public/Render/DisplayDevice/IDisplayClusterDisplayDeviceProxy.h
--- a/Engine/Plugins/Runtime/nDisplay/Source/DisplayCluster/Public/Render/DisplayDevice/IDisplayClusterDisplayDeviceProxy.h
+++ b/Engine/Plugins/Runtime/nDisplay/Source/DisplayCluster/Public/Render/DisplayDevice/IDisplayClusterDisplayDeviceProxy.h
@@ -45,7 +45,7 @@ public:
 	 * @return - true if success.
 	 */
 	virtual bool AddFinalPass_RenderThread(FRDGBuilder& GraphBuilder, const FDisplayClusterViewport_Context& InViewportContext,
-		FRHITexture2D* InputTextureRHI, const FIntRect& InputRect, FRHITexture2D* OutputTextureRHI, const FIntRect& OutputRect) const
+		FRHITexture* InputTextureRHI, const FIntRect& InputRect, FRHITexture* OutputTextureRHI, const FIntRect& OutputRect) const
 	{ return false; }
 
 	/** Allow callback OnPostProcessPassAfterFXAA.  */
diff --git a/Engine/Plugins/Runtime/nDisplay/Source/DisplayCluster/Public/Render/PostProcess/IDisplayClusterPostProcess.h b/Engine/Plugins/Runtime/nDisplay/Source/DisplayCluster/Public/Render/PostProcess/IDisplayClusterPostProcess.h
--- a/Engine/Plugins/Runtime/nDisplay/Source/DisplayCluster/Public/Render/PostProcess/IDisplayClusterPostProcess.h
+++ b/Engine/Plugins/Runtime/nDisplay/Source/DisplayCluster/Public/Render/PostProcess/IDisplayClusterPostProcess.h
@@ -170,6 +170,6 @@ public:
 	* @param RHICmdList - RHI command list
 	* @param FrameTargets - Frame textures array (1 for mono, 2 for stereo)
 	*/
-	virtual void PerformPostProcessFrameAfterWarpBlend_RenderThread(FRHICommandListImmediate& RHICmdList, const TArray<FRHITexture2D*>* InFrameTargets = nullptr, const TArray<FRHITexture2D*>* InAdditionalFrameTargets = nullptr) const
+	virtual void PerformPostProcessFrameAfterWarpBlend_RenderThread(FRHICommandListImmediate& RHICmdList, const TArray<FRHITexture*>* InFrameTargets = nullptr, const TArray<FRHITexture*>* InAdditionalFrameTargets = nullptr) const
 	{ }
 };
diff --git a/Engine/Plugins/Runtime/nDisplay/Source/DisplayCluster/Public/Render/Viewport/IDisplayClusterViewportManagerProxy.h b/Engine/Plugins/Runtime/nDisplay/Source/DisplayCluster/Public/Render/Viewport/IDisplayClusterViewportManagerProxy.h
--- a/Engine/Plugins/Runtime/nDisplay/Source/DisplayCluster/Public/Render/Viewport/IDisplayClusterViewportManagerProxy.h
+++ b/Engine/Plugins/Runtime/nDisplay/Source/DisplayCluster/Public/Render/Viewport/IDisplayClusterViewportManagerProxy.h
@@ -60,7 +60,7 @@ public:
 	*
 	* @return - true if success
 	*/
-	virtual bool GetFrameTargets_RenderThread(TArray<FRHITexture2D*>& OutFrameResources, TArray<FIntPoint>& OutTargetOffsets, TArray<FRHITexture2D*>* OutAdditionalFrameResources=nullptr) const = 0;
+	virtual bool GetFrameTargets_RenderThread(TArray<FRHITexture*>& OutFrameResources, TArray<FIntPoint>& OutTargetOffsets, TArray<FRHITexture*>* OutAdditionalFrameResources=nullptr) const = 0;
 
 	/**
 	* Resolve to backbuffer
@@ -72,7 +72,7 @@ public:
 	*
 	* @return - true if success
 	*/
-	virtual bool ResolveFrameTargetToBackBuffer_RenderThread(FRHICommandListImmediate& RHICmdList, const uint32 InContextNum, const int32 DestArrayIndex, FRHITexture2D* DstBackBuffer, FVector2D WindowSize) const = 0;
+	virtual bool ResolveFrameTargetToBackBuffer_RenderThread(FRHICommandListImmediate& RHICmdList, const uint32 InContextNum, const int32 DestArrayIndex, FRHITexture* DstBackBuffer, FVector2D WindowSize) const = 0;
 
 	///////////////// UE_DEPRECATED 5.4 ///////////////////
 
diff --git a/Engine/Plugins/Runtime/nDisplay/Source/DisplayCluster/Public/Render/Viewport/IDisplayClusterViewportProxy.h b/Engine/Plugins/Runtime/nDisplay/Source/DisplayCluster/Public/Render/Viewport/IDisplayClusterViewportProxy.h
--- a/Engine/Plugins/Runtime/nDisplay/Source/DisplayCluster/Public/Render/Viewport/IDisplayClusterViewportProxy.h
+++ b/Engine/Plugins/Runtime/nDisplay/Source/DisplayCluster/Public/Render/Viewport/IDisplayClusterViewportProxy.h
@@ -61,7 +61,7 @@ public:
 	 *
 	 * @return - true if success
 	 */
-	virtual bool GetResources_RenderThread(const EDisplayClusterViewportResourceType InResourceType, TArray<FRHITexture2D*>& OutResources) const = 0;
+	virtual bool GetResources_RenderThread(const EDisplayClusterViewportResourceType InResourceType, TArray<FRHITexture*>& OutResources) const = 0;
 
 	/** Get viewport resources with rects by type
 	 *
@@ -71,7 +71,7 @@ public:
 	 *
 	 * @return - true if success
 	 */
-	virtual bool GetResourcesWithRects_RenderThread(const EDisplayClusterViewportResourceType InResourceType, TArray<FRHITexture2D*>& OutResources, TArray<FIntRect>& OutRects) const = 0;
+	virtual bool GetResourcesWithRects_RenderThread(const EDisplayClusterViewportResourceType InResourceType, TArray<FRHITexture*>& OutResources, TArray<FIntRect>& OutRects) const = 0;
 
 	/** Copy resource contexts by type
 	 *
diff --git a/Engine/Plugins/Runtime/nDisplay/Source/DisplayClusterMoviePipeline/Private/DisplayClusterMoviePipelineViewportPass.cpp b/Engine/Plugins/Runtime/nDisplay/Source/DisplayClusterMoviePipeline/Private/DisplayClusterMoviePipelineViewportPass.cpp
--- a/Engine/Plugins/Runtime/nDisplay/Source/DisplayClusterMoviePipeline/Private/DisplayClusterMoviePipelineViewportPass.cpp
+++ b/Engine/Plugins/Runtime/nDisplay/Source/DisplayClusterMoviePipeline/Private/DisplayClusterMoviePipelineViewportPass.cpp
@@ -32,7 +32,7 @@ namespace UE::DisplayCluster::MoviePipelineViewportPass
 				const TSharedPtr<IDisplayClusterProjectionPolicy, ESPMode::ThreadSafe>& InPolicy = InViewportProxy->GetProjectionPolicy_RenderThread();
 				if (InPolicy.IsValid() && InPolicy->IsWarpBlendSupported())
 				{
-					TArray<FRHITexture2D*> WarpInputTextures, WarpOutputTextures;
+					TArray<FRHITexture*> WarpInputTextures, WarpOutputTextures;
 					TArray<FIntRect> WarpInputRects, WarpOutputRects;
 
 					if (InViewportProxy->GetResourcesWithRects_RenderThread(EDisplayClusterViewportResourceType::InputShaderResource, WarpInputTextures, WarpInputRects))
diff --git a/Engine/Plugins/Runtime/nDisplay/Source/DisplayClusterProjection/Private/Policy/Domeprojection/Windows/DX11/DisplayClusterProjectionDomeprojectionViewAdapterDX11.cpp b/Engine/Plugins/Runtime/nDisplay/Source/DisplayClusterProjection/Private/Policy/Domeprojection/Windows/DX11/DisplayClusterProjectionDomeprojectionViewAdapterDX11.cpp
--- a/Engine/Plugins/Runtime/nDisplay/Source/DisplayClusterProjection/Private/Policy/Domeprojection/Windows/DX11/DisplayClusterProjectionDomeprojectionViewAdapterDX11.cpp
+++ b/Engine/Plugins/Runtime/nDisplay/Source/DisplayClusterProjection/Private/Policy/Domeprojection/Windows/DX11/DisplayClusterProjectionDomeprojectionViewAdapterDX11.cpp
@@ -122,7 +122,7 @@ bool FDisplayClusterProjectionDomeprojectionViewAdapterDX11::GetProjectionMatrix
 	return true;
 }
 
-bool FDisplayClusterProjectionDomeprojectionViewAdapterDX11::ImplApplyWarpBlend_RenderThread(FRHICommandListImmediate& RHICmdList, uint32 InContextNum, const uint32 Channel, FRHITexture2D* InputTextures, FRHITexture2D* OutputTextures)
+bool FDisplayClusterProjectionDomeprojectionViewAdapterDX11::ImplApplyWarpBlend_RenderThread(FRHICommandListImmediate& RHICmdList, uint32 InContextNum, const uint32 Channel, FRHITexture* InputTextures, FRHITexture* OutputTextures)
 {
 	if (!IsRHID3D11())
 	{
@@ -187,7 +187,7 @@ bool FDisplayClusterProjectionDomeprojectionViewAdapterDX11::ApplyWarpBlend_Rend
 	}
 
 	// Get in\out remp resources ref from viewport
-	TArray<FRHITexture2D*> InputTextures, OutputTextures;
+	TArray<FRHITexture*> InputTextures, OutputTextures;
 
 	if (!InViewportProxy->GetResources_RenderThread(EDisplayClusterViewportResourceType::InputShaderResource, InputTextures))
 	{
diff --git a/Engine/Plugins/Runtime/nDisplay/Source/DisplayClusterProjection/Private/Policy/Domeprojection/Windows/DX11/DisplayClusterProjectionDomeprojectionViewAdapterDX11.h b/Engine/Plugins/Runtime/nDisplay/Source/DisplayClusterProjection/Private/Policy/Domeprojection/Windows/DX11/DisplayClusterProjectionDomeprojectionViewAdapterDX11.h
--- a/Engine/Plugins/Runtime/nDisplay/Source/DisplayClusterProjection/Private/Policy/Domeprojection/Windows/DX11/DisplayClusterProjectionDomeprojectionViewAdapterDX11.h
+++ b/Engine/Plugins/Runtime/nDisplay/Source/DisplayClusterProjection/Private/Policy/Domeprojection/Windows/DX11/DisplayClusterProjectionDomeprojectionViewAdapterDX11.h
@@ -31,7 +31,7 @@ public:
 	virtual bool ApplyWarpBlend_RenderThread(FRHICommandListImmediate& RHICmdList, const IDisplayClusterViewportProxy* InViewportProxy, const uint32 Channel) override;
 
 protected:
-	bool ImplApplyWarpBlend_RenderThread(FRHICommandListImmediate& RHICmdList, uint32 ContextNum, const uint32 Channel, FRHITexture2D* InputTextures, FRHITexture2D* OutputTextures);
+	bool ImplApplyWarpBlend_RenderThread(FRHICommandListImmediate& RHICmdList, uint32 ContextNum, const uint32 Channel, FRHITexture* InputTextures, FRHITexture* OutputTextures);
 
 private:
 	float ZNear;
diff --git a/Engine/Plugins/Runtime/nDisplay/Source/DisplayClusterProjection/Private/Policy/EasyBlend/DisplayClusterProjectionEasyBlendPolicy.cpp b/Engine/Plugins/Runtime/nDisplay/Source/DisplayClusterProjection/Private/Policy/EasyBlend/DisplayClusterProjectionEasyBlendPolicy.cpp
--- a/Engine/Plugins/Runtime/nDisplay/Source/DisplayClusterProjection/Private/Policy/EasyBlend/DisplayClusterProjectionEasyBlendPolicy.cpp
+++ b/Engine/Plugins/Runtime/nDisplay/Source/DisplayClusterProjection/Private/Policy/EasyBlend/DisplayClusterProjectionEasyBlendPolicy.cpp
@@ -224,7 +224,7 @@ void FDisplayClusterProjectionEasyBlendPolicy::ApplyWarpBlend_RenderThread(FRHIC
 	}
 
 	// Get in\out remp resources ref from viewport
-	TArray<FRHITexture2D*> InputTextures, OutputTextures;
+	TArray<FRHITexture*> InputTextures, OutputTextures;
 	if (InViewportProxy->GetResources_RenderThread(EDisplayClusterViewportResourceType::InputShaderResource, InputTextures)
 		&& InViewportProxy->GetResources_RenderThread(EDisplayClusterViewportResourceType::AdditionalTargetableResource, OutputTextures))
 	{
diff --git a/Engine/Plugins/Runtime/nDisplay/Source/DisplayClusterProjection/Private/Policy/EasyBlend/IDisplayClusterProjectionEasyBlendPolicyViewData.h b/Engine/Plugins/Runtime/nDisplay/Source/DisplayClusterProjection/Private/Policy/EasyBlend/IDisplayClusterProjectionEasyBlendPolicyViewData.h
--- a/Engine/Plugins/Runtime/nDisplay/Source/DisplayClusterProjection/Private/Policy/EasyBlend/IDisplayClusterProjectionEasyBlendPolicyViewData.h
+++ b/Engine/Plugins/Runtime/nDisplay/Source/DisplayClusterProjection/Private/Policy/EasyBlend/IDisplayClusterProjectionEasyBlendPolicyViewData.h
@@ -37,7 +37,7 @@ public:
 	virtual bool Initialize(const FDisplayClusterProjectionEasyBlendPolicyConfiguration& InEasyBlendConfiguration) { return false; }
 
 	virtual bool CalculateWarpBlend(FDisplayClusterProjectionEasyBlendPolicyViewInfo& InOutViewInfo) { return false; }
-	virtual bool ApplyWarpBlend_RenderThread(FRHICommandListImmediate& RHICmdList, const FDisplayClusterProjectionEasyBlendPolicyViewInfo& InViewInfo, FRHITexture2D* InputTexture, FRHITexture2D* OutputTexture, FRHIViewport* InRHIViewport) { return false; }
+	virtual bool ApplyWarpBlend_RenderThread(FRHICommandListImmediate& RHICmdList, const FDisplayClusterProjectionEasyBlendPolicyViewInfo& InViewInfo, FRHITexture* InputTexture, FRHITexture* OutputTexture, FRHIViewport* InRHIViewport) { return false; }
 
 	/**
 	* Export warp mesh geometry from EasyBlend file
diff --git a/Engine/Plugins/Runtime/nDisplay/Source/DisplayClusterProjection/Private/Policy/EasyBlend/Windows/DX11/DisplayClusterProjectionEasyBlendPolicyViewDataDX11.cpp b/Engine/Plugins/Runtime/nDisplay/Source/DisplayClusterProjection/Private/Policy/EasyBlend/Windows/DX11/DisplayClusterProjectionEasyBlendPolicyViewDataDX11.cpp
--- a/Engine/Plugins/Runtime/nDisplay/Source/DisplayClusterProjection/Private/Policy/EasyBlend/Windows/DX11/DisplayClusterProjectionEasyBlendPolicyViewDataDX11.cpp
+++ b/Engine/Plugins/Runtime/nDisplay/Source/DisplayClusterProjection/Private/Policy/EasyBlend/Windows/DX11/DisplayClusterProjectionEasyBlendPolicyViewDataDX11.cpp
@@ -115,7 +115,7 @@ bool FDisplayClusterProjectionEasyBlendPolicyViewDataDX11::CalculateWarpBlend(FD
 	return true;
 }
 
-bool FDisplayClusterProjectionEasyBlendPolicyViewDataDX11::ApplyWarpBlend_RenderThread(FRHICommandListImmediate& RHICmdList, const FDisplayClusterProjectionEasyBlendPolicyViewInfo& InViewInfo, FRHITexture2D* InputTexture, FRHITexture2D* OutputTexture, FRHIViewport* InRHIViewport)
+bool FDisplayClusterProjectionEasyBlendPolicyViewDataDX11::ApplyWarpBlend_RenderThread(FRHICommandListImmediate& RHICmdList, const FDisplayClusterProjectionEasyBlendPolicyViewInfo& InViewInfo, FRHITexture* InputTexture, FRHITexture* OutputTexture, FRHIViewport* InRHIViewport)
 {
 	check(IsInRenderingThread());
 
diff --git a/Engine/Plugins/Runtime/nDisplay/Source/DisplayClusterProjection/Private/Policy/EasyBlend/Windows/DX11/DisplayClusterProjectionEasyBlendPolicyViewDataDX11.h b/Engine/Plugins/Runtime/nDisplay/Source/DisplayClusterProjection/Private/Policy/EasyBlend/Windows/DX11/DisplayClusterProjectionEasyBlendPolicyViewDataDX11.h
--- a/Engine/Plugins/Runtime/nDisplay/Source/DisplayClusterProjection/Private/Policy/EasyBlend/Windows/DX11/DisplayClusterProjectionEasyBlendPolicyViewDataDX11.h
+++ b/Engine/Plugins/Runtime/nDisplay/Source/DisplayClusterProjection/Private/Policy/EasyBlend/Windows/DX11/DisplayClusterProjectionEasyBlendPolicyViewDataDX11.h
@@ -20,7 +20,7 @@ public:
 	virtual bool Initialize(const FDisplayClusterProjectionEasyBlendPolicyConfiguration& InEasyBlendConfiguration) override;
 
 	virtual bool CalculateWarpBlend(FDisplayClusterProjectionEasyBlendPolicyViewInfo& InOutViewInfo) override;
-	virtual bool ApplyWarpBlend_RenderThread(FRHICommandListImmediate& RHICmdList, const FDisplayClusterProjectionEasyBlendPolicyViewInfo& InViewInfo, FRHITexture2D* InputTexture, FRHITexture2D* OutputTexture, FRHIViewport* InRHIViewport) override;
+	virtual bool ApplyWarpBlend_RenderThread(FRHICommandListImmediate& RHICmdList, const FDisplayClusterProjectionEasyBlendPolicyViewInfo& InViewInfo, FRHITexture* InputTexture, FRHITexture* OutputTexture, FRHIViewport* InRHIViewport) override;
 	//END ~~IDisplayClusterProjectionEasyBlendPolicyViewData
 
 protected:
diff --git a/Engine/Plugins/Runtime/nDisplay/Source/DisplayClusterProjection/Private/Policy/EasyBlend/Windows/DX12/DisplayClusterProjectionEasyBlendPolicyViewDataDX12.cpp b/Engine/Plugins/Runtime/nDisplay/Source/DisplayClusterProjection/Private/Policy/EasyBlend/Windows/DX12/DisplayClusterProjectionEasyBlendPolicyViewDataDX12.cpp
--- a/Engine/Plugins/Runtime/nDisplay/Source/DisplayClusterProjection/Private/Policy/EasyBlend/Windows/DX12/DisplayClusterProjectionEasyBlendPolicyViewDataDX12.cpp
+++ b/Engine/Plugins/Runtime/nDisplay/Source/DisplayClusterProjection/Private/Policy/EasyBlend/Windows/DX12/DisplayClusterProjectionEasyBlendPolicyViewDataDX12.cpp
@@ -140,7 +140,7 @@ bool FDisplayClusterProjectionEasyBlendPolicyViewDataDX12::CalculateWarpBlend(FD
 	return true;
 }
 
-bool FDisplayClusterProjectionEasyBlendPolicyViewDataDX12::ApplyWarpBlend_RenderThread(FRHICommandListImmediate& RHICmdList, const FDisplayClusterProjectionEasyBlendPolicyViewInfo& InViewInfo, FRHITexture2D* InputTexture, FRHITexture2D* OutputTexture, FRHIViewport* InRHIViewport)
+bool FDisplayClusterProjectionEasyBlendPolicyViewDataDX12::ApplyWarpBlend_RenderThread(FRHICommandListImmediate& RHICmdList, const FDisplayClusterProjectionEasyBlendPolicyViewInfo& InViewInfo, FRHITexture* InputTexture, FRHITexture* OutputTexture, FRHIViewport* InRHIViewport)
 {
 	check(IsInRenderingThread());
 
diff --git a/Engine/Plugins/Runtime/nDisplay/Source/DisplayClusterProjection/Private/Policy/EasyBlend/Windows/DX12/DisplayClusterProjectionEasyBlendPolicyViewDataDX12.h b/Engine/Plugins/Runtime/nDisplay/Source/DisplayClusterProjection/Private/Policy/EasyBlend/Windows/DX12/DisplayClusterProjectionEasyBlendPolicyViewDataDX12.h
--- a/Engine/Plugins/Runtime/nDisplay/Source/DisplayClusterProjection/Private/Policy/EasyBlend/Windows/DX12/DisplayClusterProjectionEasyBlendPolicyViewDataDX12.h
+++ b/Engine/Plugins/Runtime/nDisplay/Source/DisplayClusterProjection/Private/Policy/EasyBlend/Windows/DX12/DisplayClusterProjectionEasyBlendPolicyViewDataDX12.h
@@ -19,7 +19,7 @@ public:
 	//BEGIN ~IDisplayClusterProjectionEasyBlendPolicyViewData
 	virtual bool Initialize(const FDisplayClusterProjectionEasyBlendPolicyConfiguration& InEasyBlendConfiguration) override;
 	virtual bool CalculateWarpBlend(FDisplayClusterProjectionEasyBlendPolicyViewInfo& InOutViewInfo) override;
-	virtual bool ApplyWarpBlend_RenderThread(FRHICommandListImmediate& RHICmdList, const FDisplayClusterProjectionEasyBlendPolicyViewInfo& InViewInfo, FRHITexture2D* InputTexture, FRHITexture2D* OutputTexture, FRHIViewport* InRHIViewport) override;
+	virtual bool ApplyWarpBlend_RenderThread(FRHICommandListImmediate& RHICmdList, const FDisplayClusterProjectionEasyBlendPolicyViewInfo& InViewInfo, FRHITexture* InputTexture, FRHITexture* OutputTexture, FRHIViewport* InRHIViewport) override;
 
 	virtual bool HasPreviewMesh() override { return true; }
 	virtual bool GetPreviewMeshGeometry(const FDisplayClusterProjectionEasyBlendPolicyConfiguration& InEasyBlendConfiguration, FDisplayClusterProjectionEasyBlendGeometryExportData& OutMeshData) override;
diff --git a/Engine/Plugins/Runtime/nDisplay/Source/DisplayClusterProjection/Private/Policy/MPCDI/DisplayClusterProjectionMPCDIPolicy.cpp b/Engine/Plugins/Runtime/nDisplay/Source/DisplayClusterProjection/Private/Policy/MPCDI/DisplayClusterProjectionMPCDIPolicy.cpp
--- a/Engine/Plugins/Runtime/nDisplay/Source/DisplayClusterProjection/Private/Policy/MPCDI/DisplayClusterProjectionMPCDIPolicy.cpp
+++ b/Engine/Plugins/Runtime/nDisplay/Source/DisplayClusterProjection/Private/Policy/MPCDI/DisplayClusterProjectionMPCDIPolicy.cpp
@@ -328,7 +328,7 @@ void FDisplayClusterProjectionMPCDIPolicy::ApplyWarpBlend_RenderThread(FRHIComma
 		return;
 	}
 
-	TArray<FRHITexture2D*> InputTextures, OutputTextures;
+	TArray<FRHITexture*> InputTextures, OutputTextures;
 	TArray<FIntRect> InputRects, OutputRects;
 
 	// Use for input first MipsShader texture if enabled in viewport render settings
@@ -367,7 +367,7 @@ void FDisplayClusterProjectionMPCDIPolicy::ApplyWarpBlend_RenderThread(FRHIComma
 				{
 					if (!SrcViewportProxy->GetRenderSettings_RenderThread().bSkipRendering)
 					{
-						TArray<FRHITexture2D*> RefTextures;
+						TArray<FRHITexture*> RefTextures;
 						// Use for input first MipsShader texture if enabled in viewport render settings
 						if (SrcViewportProxy->GetResources_RenderThread(EDisplayClusterViewportResourceType::MipsShaderResource, RefTextures) ||
 							SrcViewportProxy->GetResources_RenderThread(EDisplayClusterViewportResourceType::InputShaderResource, RefTextures))
diff --git a/Engine/Plugins/Runtime/nDisplay/Source/DisplayClusterProjection/Private/Policy/VIOSO/Windows/DisplayClusterProjectionVIOSOPolicy.cpp b/Engine/Plugins/Runtime/nDisplay/Source/DisplayClusterProjection/Private/Policy/VIOSO/Windows/DisplayClusterProjectionVIOSOPolicy.cpp
--- a/Engine/Plugins/Runtime/nDisplay/Source/DisplayClusterProjection/Private/Policy/VIOSO/Windows/DisplayClusterProjectionVIOSOPolicy.cpp
+++ b/Engine/Plugins/Runtime/nDisplay/Source/DisplayClusterProjection/Private/Policy/VIOSO/Windows/DisplayClusterProjectionVIOSOPolicy.cpp
@@ -165,7 +165,7 @@ bool FDisplayClusterProjectionVIOSOPolicy::ImplApplyWarpBlend_RenderThread(FRHIC
 	check(IsInRenderingThread());
 
 	// Get in\out remp resources ref from viewport
-	TArray<FRHITexture2D*> InputTextures, OutputTextures;
+	TArray<FRHITexture*> InputTextures, OutputTextures;
 
 	// Use for input first MipsShader texture if enabled in viewport render settings
 	//@todo: test if domeprojection support mips textures as warp input
diff --git a/Engine/Plugins/Runtime/nDisplay/Source/DisplayClusterProjection/Private/Policy/VIOSO/Windows/DisplayClusterProjectionVIOSOPolicyViewData.cpp b/Engine/Plugins/Runtime/nDisplay/Source/DisplayClusterProjection/Private/Policy/VIOSO/Windows/DisplayClusterProjectionVIOSOPolicyViewData.cpp
--- a/Engine/Plugins/Runtime/nDisplay/Source/DisplayClusterProjection/Private/Policy/VIOSO/Windows/DisplayClusterProjectionVIOSOPolicyViewData.cpp
+++ b/Engine/Plugins/Runtime/nDisplay/Source/DisplayClusterProjection/Private/Policy/VIOSO/Windows/DisplayClusterProjectionVIOSOPolicyViewData.cpp
@@ -91,7 +91,7 @@ public:
 		}
 	}
 
-	inline bool AssignD3D11RenderTarget(FRHITexture2D* RenderTargetTexture)
+	inline bool AssignD3D11RenderTarget(FRHITexture* RenderTargetTexture)
 	{
 		if (DeviceContext)
 		{
@@ -186,7 +186,7 @@ bool FDisplayClusterProjectionVIOSOPolicyViewData::UpdateVIOSO(IDisplayClusterVi
 	return false;
 }
 
-bool FDisplayClusterProjectionVIOSOPolicyViewData::RenderVIOSO_RenderThread(FRHICommandListImmediate& RHICmdList, FRHITexture2D* ShaderResourceTexture, FRHITexture2D* RenderTargetTexture)
+bool FDisplayClusterProjectionVIOSOPolicyViewData::RenderVIOSO_RenderThread(FRHICommandListImmediate& RHICmdList, FRHITexture* ShaderResourceTexture, FRHITexture* RenderTargetTexture)
 {
 	check(IsInRenderingThread());
 
@@ -255,7 +255,7 @@ bool FDisplayClusterProjectionVIOSOPolicyViewData::RenderVIOSO_RenderThread(FRHI
 	return false;
 }
 
-bool FDisplayClusterProjectionVIOSOPolicyViewData::InitializeVIOSO_RenderThread(FRHICommandListImmediate& RHICmdList, FRHITexture2D* RenderTargetTexture)
+bool FDisplayClusterProjectionVIOSOPolicyViewData::InitializeVIOSO_RenderThread(FRHICommandListImmediate& RHICmdList, FRHITexture* RenderTargetTexture)
 {
 #if WITH_VIOSO_LIBRARY
 
diff --git a/Engine/Plugins/Runtime/nDisplay/Source/DisplayClusterProjection/Private/Policy/VIOSO/Windows/DisplayClusterProjectionVIOSOPolicyViewData.h b/Engine/Plugins/Runtime/nDisplay/Source/DisplayClusterProjection/Private/Policy/VIOSO/Windows/DisplayClusterProjectionVIOSOPolicyViewData.h
--- a/Engine/Plugins/Runtime/nDisplay/Source/DisplayClusterProjection/Private/Policy/VIOSO/Windows/DisplayClusterProjectionVIOSOPolicyViewData.h
+++ b/Engine/Plugins/Runtime/nDisplay/Source/DisplayClusterProjection/Private/Policy/VIOSO/Windows/DisplayClusterProjectionVIOSOPolicyViewData.h
@@ -38,10 +38,10 @@ public:
 	}
 
 	bool UpdateVIOSO(IDisplayClusterViewport* InViewport, const uint32 InContextNum, const FVector& LocalLocation, const FRotator& LocalRotator, const float WorldToMeters, const float NCP, const float FCP);
-	bool RenderVIOSO_RenderThread(FRHICommandListImmediate& RHICmdList, FRHITexture2D* ShaderResourceTexture, FRHITexture2D* RenderTargetTexture);
+	bool RenderVIOSO_RenderThread(FRHICommandListImmediate& RHICmdList, FRHITexture* ShaderResourceTexture, FRHITexture* RenderTargetTexture);
 
 protected:
-	bool InitializeVIOSO_RenderThread(FRHICommandListImmediate& RHICmdList, FRHITexture2D* RenderTargetTexture);
+	bool InitializeVIOSO_RenderThread(FRHICommandListImmediate& RHICmdList, FRHITexture* RenderTargetTexture);
 
 public:
 	FVector  ViewLocation;
@@ -59,7 +59,7 @@ private:
 	bool bInitialized = false;
 
 	// This RTT is used for VIOSO. If we change this resource, we need to re-initialize VIOSO
-	FRHITexture2D* UsedRenderTargetTexture = nullptr;
+	FRHITexture* UsedRenderTargetTexture = nullptr;
 
 	FVector2D ClippingPlanes;
 };
diff --git a/Engine/Plugins/Runtime/nDisplay/Source/DisplayClusterShaders/Private/DisplayClusterShadersModule.cpp b/Engine/Plugins/Runtime/nDisplay/Source/DisplayClusterShaders/Private/DisplayClusterShadersModule.cpp
--- a/Engine/Plugins/Runtime/nDisplay/Source/DisplayClusterShaders/Private/DisplayClusterShadersModule.cpp
+++ b/Engine/Plugins/Runtime/nDisplay/Source/DisplayClusterShaders/Private/DisplayClusterShadersModule.cpp
@@ -47,17 +47,17 @@ bool FDisplayClusterShadersModule::RenderPreprocess_UVLightCards(FRHICommandList
 	return FDisplayClusterShadersPreprocess_UVLightCards::RenderPreprocess_UVLightCards(RHICmdList, InScene, InRenderTarget, InParameters);
 }
 
-bool FDisplayClusterShadersModule::RenderPostprocess_OutputRemap(FRHICommandListImmediate& RHICmdList, FRHITexture2D* InSourceTexture, FRHITexture2D* InRenderTargetableDestTexture, const IDisplayClusterRender_MeshComponentProxy& MeshProxy) const
+bool FDisplayClusterShadersModule::RenderPostprocess_OutputRemap(FRHICommandListImmediate& RHICmdList, FRHITexture* InSourceTexture, FRHITexture* InRenderTargetableDestTexture, const IDisplayClusterRender_MeshComponentProxy& MeshProxy) const
 {
 	return FDisplayClusterShadersPostprocess_OutputRemap::RenderPostprocess_OutputRemap(RHICmdList, InSourceTexture, InRenderTargetableDestTexture, MeshProxy);
 }
 
-bool FDisplayClusterShadersModule::RenderPostprocess_Blur(FRHICommandListImmediate& RHICmdList, FRHITexture2D* InSourceTexture, FRHITexture2D* InRenderTargetableDestTexture, const FDisplayClusterShaderParameters_PostprocessBlur& InSettings) const
+bool FDisplayClusterShadersModule::RenderPostprocess_Blur(FRHICommandListImmediate& RHICmdList, FRHITexture* InSourceTexture, FRHITexture* InRenderTargetableDestTexture, const FDisplayClusterShaderParameters_PostprocessBlur& InSettings) const
 {
 	return FDisplayClusterShadersPostprocess_Blur::RenderPostprocess_Blur(RHICmdList, InSourceTexture, InRenderTargetableDestTexture, InSettings);
 }
 
-bool FDisplayClusterShadersModule::GenerateMips(FRHICommandListImmediate& RHICmdList, FRHITexture2D* InOutMipsTexture, const FDisplayClusterShaderParameters_GenerateMips& InSettings) const
+bool FDisplayClusterShadersModule::GenerateMips(FRHICommandListImmediate& RHICmdList, FRHITexture* InOutMipsTexture, const FDisplayClusterShaderParameters_GenerateMips& InSettings) const
 {
 	return FDisplayClusterShadersGenerateMips::GenerateMips(RHICmdList, InOutMipsTexture, InSettings);
 }
diff --git a/Engine/Plugins/Runtime/nDisplay/Source/DisplayClusterShaders/Private/DisplayClusterShadersModule.h b/Engine/Plugins/Runtime/nDisplay/Source/DisplayClusterShaders/Private/DisplayClusterShadersModule.h
--- a/Engine/Plugins/Runtime/nDisplay/Source/DisplayClusterShaders/Private/DisplayClusterShadersModule.h
+++ b/Engine/Plugins/Runtime/nDisplay/Source/DisplayClusterShaders/Private/DisplayClusterShadersModule.h
@@ -19,7 +19,7 @@ public:
 	virtual bool RenderWarpBlend_ICVFX(FRHICommandListImmediate& RHICmdList, const FDisplayClusterShaderParameters_WarpBlend& InWarpBlendParameters, const FDisplayClusterShaderParameters_ICVFX& InICVFXParameters) const override;
 	virtual bool RenderPreprocess_UVLightCards(FRHICommandListImmediate& RHICmdList, FSceneInterface* InScene, FRenderTarget* InRenderTarget, float ProjectionPlaneSize, bool bRenderFinalColor) const override { return false; }
 	virtual bool RenderPreprocess_UVLightCards(FRHICommandListImmediate& RHICmdList, FSceneInterface* InScene, FRenderTarget* InRenderTarget, const FDisplayClusterShaderParameters_UVLightCards& InParameters) const override;
-	virtual bool RenderPostprocess_OutputRemap(FRHICommandListImmediate& RHICmdList, FRHITexture2D* InSourceTexture, FRHITexture2D* InRenderTargetableDestTexture, const IDisplayClusterRender_MeshComponentProxy& MeshProxy) const override;
-	virtual bool RenderPostprocess_Blur(FRHICommandListImmediate& RHICmdList, FRHITexture2D* InSourceTexture, FRHITexture2D* InRenderTargetableDestTexture, const FDisplayClusterShaderParameters_PostprocessBlur& InSettings) const override;
-	virtual bool GenerateMips(FRHICommandListImmediate& RHICmdList, FRHITexture2D* InOutMipsTexture, const FDisplayClusterShaderParameters_GenerateMips& InSettings) const override;
+	virtual bool RenderPostprocess_OutputRemap(FRHICommandListImmediate& RHICmdList, FRHITexture* InSourceTexture, FRHITexture* InRenderTargetableDestTexture, const IDisplayClusterRender_MeshComponentProxy& MeshProxy) const override;
+	virtual bool RenderPostprocess_Blur(FRHICommandListImmediate& RHICmdList, FRHITexture* InSourceTexture, FRHITexture* InRenderTargetableDestTexture, const FDisplayClusterShaderParameters_PostprocessBlur& InSettings) const override;
+	virtual bool GenerateMips(FRHICommandListImmediate& RHICmdList, FRHITexture* InOutMipsTexture, const FDisplayClusterShaderParameters_GenerateMips& InSettings) const override;
 };
diff --git a/Engine/Plugins/Runtime/nDisplay/Source/DisplayClusterShaders/Private/Shaders/DisplayClusterShadersGenerateMips.cpp b/Engine/Plugins/Runtime/nDisplay/Source/DisplayClusterShaders/Private/Shaders/DisplayClusterShadersGenerateMips.cpp
--- a/Engine/Plugins/Runtime/nDisplay/Source/DisplayClusterShaders/Private/Shaders/DisplayClusterShadersGenerateMips.cpp
+++ b/Engine/Plugins/Runtime/nDisplay/Source/DisplayClusterShaders/Private/Shaders/DisplayClusterShadersGenerateMips.cpp
@@ -7,7 +7,7 @@
 #include "RenderGraphBuilder.h"
 #include "RenderGraphUtils.h"
 
-bool FDisplayClusterShadersGenerateMips::GenerateMips(FRHICommandListImmediate& RHICmdList, FRHITexture2D* InOutMipsTexture, const FDisplayClusterShaderParameters_GenerateMips& InSettings)
+bool FDisplayClusterShadersGenerateMips::GenerateMips(FRHICommandListImmediate& RHICmdList, FRHITexture* InOutMipsTexture, const FDisplayClusterShaderParameters_GenerateMips& InSettings)
 {
 	if (InSettings.IsEnabled())
 	{
