Add a new function to IMotionController.h, SetEnhancedInputUserSettings, to accept a UEnhancedInputUserSettings object and apply its input mappings to the controller. This replaces the deprecated SetPlayerMappableInputConfig function.

Prevent OpenXR input from binding enhanced input default mappiing contexts if "Enable Default Mapping Contexts" is unchceked.

diff --git a/Engine/Plugins/Runtime/OpenXR/Source/OpenXRInput/Private/OpenXRInput.cpp b/Engine/Plugins/Runtime/OpenXR/Source/OpenXRInput/Private/OpenXRInput.cpp
--- a/Engine/Plugins/Runtime/OpenXR/Source/OpenXRInput/Private/OpenXRInput.cpp
+++ b/Engine/Plugins/Runtime/OpenXR/Source/OpenXRInput/Private/OpenXRInput.cpp
@@ -17,6 +17,7 @@
 #include "InputAction.h"
 #include "InputMappingContext.h"
 #include "PlayerMappableInputConfig.h"
+#include "UserSettings/EnhancedInputUserSettings.h"
 #include "GenericPlatform/GenericPlatformInputDeviceMapper.h"
 #include "Engine/Engine.h"
 
@@ -404,7 +405,7 @@ bool FOpenXRInputPlugin::FOpenXRInput::BuildActions(XrSession Session)
 	PRAGMA_DISABLE_DEPRECATION_WARNINGS
 	// Attempt to load the default input config from the OpenXR input settings.
 	const UEnhancedInputDeveloperSettings* InputSettings = GetDefault<UEnhancedInputDeveloperSettings>();
-	if (InputSettings)
+	if (InputSettings && InputSettings->bEnableDefaultMappingContexts)
 	{
 		for (const auto& Context : InputSettings->DefaultMappingContexts)
 		{
@@ -1478,6 +1479,24 @@ bool FOpenXRInputPlugin::FOpenXRInput::SetPlayerMappableInputConfig(TObjectPtr<c
 }
 PRAGMA_ENABLE_DEPRECATION_WARNINGS
 
+bool FOpenXRInputPlugin::FOpenXRInput::SetEnhancedInputUserSettings(TObjectPtr<class UEnhancedInputUserSettings> InputSettings)
+{
+	if (bActionsAttached)
+	{
+		UE_LOG(LogHMD, Error, TEXT("Attempted to attach a set of enhanced user input settings when actions are already attached for the current session."));
+
+		return false;
+	}
+
+	const TSet<TObjectPtr<const UInputMappingContext>>& MappingContexts = InputSettings->GetRegisteredInputMappingContexts();
+
+	for (const auto& Context : MappingContexts)
+	{
+		InputMappingContextToPriorityMap.Add(TStrongObjectPtr<const UInputMappingContext>(Context), 0);
+	}
+	return true;
+}
+
 bool FOpenXRInputPlugin::FOpenXRInput::AttachInputMappingContexts(const TSet<TObjectPtr<UInputMappingContext>>& MappingContexts)
 {
 	if (bActionsAttached)
diff --git a/Engine/Plugins/Runtime/OpenXR/Source/OpenXRInput/Private/OpenXRInput.h b/Engine/Plugins/Runtime/OpenXR/Source/OpenXRInput/Private/OpenXRInput.h
--- a/Engine/Plugins/Runtime/OpenXR/Source/OpenXRInput/Private/OpenXRInput.h
+++ b/Engine/Plugins/Runtime/OpenXR/Source/OpenXRInput/Private/OpenXRInput.h
@@ -135,6 +135,7 @@ public:
 		PRAGMA_DISABLE_DEPRECATION_WARNINGS
 		virtual bool SetPlayerMappableInputConfig(TObjectPtr<class UPlayerMappableInputConfig> InputConfig) override;
 		PRAGMA_ENABLE_DEPRECATION_WARNINGS
+		virtual bool SetEnhancedInputUserSettings(TObjectPtr<class UEnhancedInputUserSettings> InputSettings) override;
 		virtual bool AttachInputMappingContexts(const TSet<TObjectPtr<UInputMappingContext>>& MappingContexts) override;
 
 		// IHapticDevice overrides
diff --git a/Engine/Source/Runtime/HeadMountedDisplay/Public/IMotionController.h b/Engine/Source/Runtime/HeadMountedDisplay/Public/IMotionController.h
--- a/Engine/Source/Runtime/HeadMountedDisplay/Public/IMotionController.h
+++ b/Engine/Source/Runtime/HeadMountedDisplay/Public/IMotionController.h
@@ -164,6 +164,16 @@ public:
 	virtual bool SetPlayerMappableInputConfig(TObjectPtr<class UPlayerMappableInputConfig> InputConfig = nullptr) { return true; };
 	PRAGMA_ENABLE_DEPRECATION_WARNINGS
 
+	/**
+	* Add a set of enhanced input user settings to the motion controller. This allows the motion controller to support
+	* Enhanced Input actions.
+	*
+	* @param InputConfig		The path to the enhanced user input settings asset
+	*
+	* @return			False if the input config can't be attached to the session, true otherwise
+	*/
+	virtual bool SetEnhancedInputUserSettings(TObjectPtr<class UEnhancedInputUserSettings> InputSettings = nullptr) { return true; };
+
 	/**
 	 * Add input mapping contexts to the motion controller. This allows the motion controller to support
 	 * Enhanced Input actions.
