Fix crash in water rendering when running with ISR on an OpenXR device (UE-207585).

diff --git a/Engine/Plugins/Runtime/OpenXR/Source/OpenXRHMD/Private/OpenXRHMD.cpp b/Engine/Plugins/Runtime/OpenXR/Source/OpenXRHMD/Private/OpenXRHMD.cpp
--- a/Engine/Plugins/Runtime/OpenXR/Source/OpenXRHMD/Private/OpenXRHMD.cpp
+++ b/Engine/Plugins/Runtime/OpenXR/Source/OpenXRHMD/Private/OpenXRHMD.cpp
@@ -1536,7 +1536,10 @@ const FOpenXRHMD::FPipelinedFrameState& FOpenXRHMD::GetPipelinedFrameStateForThr
 	// Relying on implicit selection of the RHI struct is hazardous since the RHI thread isn't always present
 	check(!IsInRHIThread());
 
-	if (IsInActualRenderingThread())
+	// Opening up access to parallel rendering threads, because some frame state (e.g. GetDesiredNumberOfViews()) is started being requested on them.
+	// Since the frame state is returned const from this function, this is hopefully a little bit more safe, but still prone to race conditions if the real
+	// render thread at this moment is modifying the state using the non-const method. Proper resolution is tracked in UE-212224.
+	if (IsInActualRenderingThread() || IsInParallelRenderingThread())
 	{
 		return PipelinedFrameStateRendering;
 	}
