UE-177257 Add CVAR to control whether we throttle the application when HMD is not worn
-xr.OpenXRPausedIdleFPS created to specify the max fps when the openxr session is idle. When set to 0 there is no max fps.  10fps is the default.  Setting it to 60 or 90 is more likely to be correct than 0, because with openxr idle the 'natural' max framerate may be faster than it is in vr.
-Also changed OnStartGameFrame to avoid some useless work when we are not in an OpenXR session.

diff --git a/Engine/Plugins/Runtime/OpenXR/Source/OpenXRHMD/Private/OpenXRHMD.cpp b/Engine/Plugins/Runtime/OpenXR/Source/OpenXRHMD/Private/OpenXRHMD.cpp
--- a/Engine/Plugins/Runtime/OpenXR/Source/OpenXRHMD/Private/OpenXRHMD.cpp
+++ b/Engine/Plugins/Runtime/OpenXR/Source/OpenXRHMD/Private/OpenXRHMD.cpp
@@ -72,9 +72,15 @@ static const TCHAR* HMDThreadString()
 
 #define LOCTEXT_NAMESPACE "OpenXR"
 
-#define OPENXR_PAUSED_IDLE_FPS 10
 static const int64 OPENXR_SWAPCHAIN_WAIT_TIMEOUT = 100000000ll;		// 100ms in nanoseconds.
 
+static TAutoConsoleVariable<int32> CVarOpenXRPausedIdleFPS(
+	TEXT("xr.OpenXRPausedIdleFPS"),
+	10,
+	TEXT("If non-zero MaxFPS will be set to this value when the XRSession state is XR_SESSION_STATE_IDLE, which often means the HMD has been removed from the users head.\n")
+	TEXT("Defaults to 10fps. 0 will allow unreal to run as fast as it can.  Note that in XR_SESSION_STATE_IDLE the frame rate may actually be higher than when in VR, so you may want to set it to 60 or 90 rather than 0.\n"),
+	ECVF_Default);
+
 static TAutoConsoleVariable<int32> CVarOpenXRExitAppOnRuntimeDrivenSessionExit(
 	TEXT("xr.OpenXRExitAppOnRuntimeDrivenSessionExit"),
 	1,
@@ -3217,14 +3223,6 @@ bool FOpenXRHMD::OnStartGameFrame(FWorldContext& WorldContext)
 	}
 #endif // WITH_EDITOR
 
-	const AWorldSettings* const WorldSettings = WorldContext.World() ? WorldContext.World()->GetWorldSettings() : nullptr;
-	if (WorldSettings)
-	{
-		WorldToMetersScale = WorldSettings->WorldToMeters;
-	}
-
-	RefreshTrackingToWorldTransform(WorldContext);
-
 	if (!System)
 	{
 		System = IOpenXRHMDModule::Get().GetSystemId();
@@ -3232,8 +3230,21 @@ bool FOpenXRHMD::OnStartGameFrame(FWorldContext& WorldContext)
 		{
 			FCoreDelegates::VRHeadsetReconnected.Broadcast();
 		}
+		else if (Session == XR_NULL_HANDLE)
+		{
+			// Having a session but no system does not make much sense, but we will continue to process XrEvents just in case.
+			return false;
+		}
+	}
+
+	const AWorldSettings* const WorldSettings = WorldContext.World() ? WorldContext.World()->GetWorldSettings() : nullptr;
+	if (WorldSettings)
+	{
+		WorldToMetersScale = WorldSettings->WorldToMeters;
 	}
 
+	RefreshTrackingToWorldTransform(WorldContext);
+
 	if (bIsTrackingOnlySession)
 	{
 		if (OnStereoStartup())
@@ -3279,7 +3290,8 @@ bool FOpenXRHMD::OnStartGameFrame(FWorldContext& WorldContext)
 			{
 				if (!GIsEditor)
 				{
-					GEngine->SetMaxFPS(OPENXR_PAUSED_IDLE_FPS);
+					const int32 PausedIdleFPS = CVarOpenXRPausedIdleFPS.GetValueOnAnyThread();
+					GEngine->SetMaxFPS(PausedIdleFPS);
 				}
 				FCoreDelegates::VRHeadsetRemovedFromHead.Broadcast();
 				bIsReady = false;
