UE-210908 Prevent IOpenXRExtensionPlugins from being added multiple times
-Clear the extensions array when initializing the OpenXR instance so that multiple-initialization does not result in plugin's being multiply added.

diff --git a/Engine/Plugins/Runtime/OpenXR/Source/OpenXRHMD/Private/OpenXRHMDModule.cpp b/Engine/Plugins/Runtime/OpenXR/Source/OpenXRHMD/Private/OpenXRHMDModule.cpp
--- a/Engine/Plugins/Runtime/OpenXR/Source/OpenXRHMD/Private/OpenXRHMDModule.cpp
+++ b/Engine/Plugins/Runtime/OpenXR/Source/OpenXRHMD/Private/OpenXRHMDModule.cpp
@@ -690,6 +690,7 @@ bool FOpenXRHMDModule::InitInstance()
 		ExtensionSet.Add(XR_EPIC_VIEW_CONFIGURATION_FOV_EXTENSION_NAME);
 	}
 
+	ExtensionPlugins.Reset();
 	for (IOpenXRExtensionPlugin* Plugin : ExtModules)
 	{
 		TArray<const ANSICHAR*> RequiredExtensions, OptionalExtensions, Extensions;
