Add locks around pipelined frame state access to make it safe to use from task threads.

- While the state is "per named thread", we're encountering scenarios when some of the functions may end up being called in "parallel rendering thread" (i.e. on task threads) so essentially we could still have race hazards.
- Such accesses will become more likely as things get parallelized, so this is an attempt to future-proof the plugin.

diff --git a/Engine/Plugins/Runtime/OpenXR/Source/OpenXRHMD/Private/OpenXRHMD.cpp b/Engine/Plugins/Runtime/OpenXR/Source/OpenXRHMD/Private/OpenXRHMD.cpp
--- a/Engine/Plugins/Runtime/OpenXR/Source/OpenXRHMD/Private/OpenXRHMD.cpp
+++ b/Engine/Plugins/Runtime/OpenXR/Source/OpenXRHMD/Private/OpenXRHMD.cpp
@@ -448,7 +448,8 @@ bool FOpenXRHMD::GetPlayAreaRect(FTransform& OutTransform, FVector2D& OutRect) c
 	// Get the origin and the extents of the play area rect.
 	// The OpenXR Stage Space defines the origin of the playable rectangle.  The origin is at the floor. xrGetReferenceSpaceBoundsRect will give you the horizontal extents.
 
-	const FPipelinedFrameState& PipelinedState = GetPipelinedFrameStateForThread();
+	FPipelinedFrameStateAccessorReadOnly LockedPipelineState = GetPipelinedFrameStateForThread();
+	const FPipelinedFrameState& PipelinedState = LockedPipelineState.GetFrameState();
 
 	{
 		if (StageSpace == XR_NULL_HANDLE)
@@ -522,7 +523,8 @@ bool FOpenXRHMD::GetTrackingOriginTransform(TEnumAsByte<EHMDTrackingOrigin::Type
 		return false;
 	}
 
-	const FPipelinedFrameState& PipelinedState = GetPipelinedFrameStateForThread();
+	FPipelinedFrameStateAccessorReadOnly LockedPipelineState = GetPipelinedFrameStateForThread();
+	const FPipelinedFrameState& PipelinedState = LockedPipelineState.GetFrameState();
 
 	if (!PipelinedState.TrackingSpace.IsValid())
 	{
@@ -592,7 +594,8 @@ bool FOpenXRHMD::GetHMDMonitorInfo(MonitorInfo& MonitorDesc)
 
 void FOpenXRHMD::GetFieldOfView(float& OutHFOVInDegrees, float& OutVFOVInDegrees) const
 {
-	const FPipelinedFrameState& FrameState = GetPipelinedFrameStateForThread();
+	FPipelinedFrameStateAccessorReadOnly LockedPipelineState = GetPipelinedFrameStateForThread();
+	const FPipelinedFrameState& FrameState = LockedPipelineState.GetFrameState();
 
 	XrFovf UnifiedFov = { 0.0f };
 	for (const XrView& View : FrameState.Views)
@@ -631,7 +634,8 @@ void FOpenXRHMD::SetInterpupillaryDistance(float NewInterpupillaryDistance)
 
 float FOpenXRHMD::GetInterpupillaryDistance() const
 {
-	const FPipelinedFrameState& FrameState = GetPipelinedFrameStateForThread();
+	FPipelinedFrameStateAccessorReadOnly LockedPipelineState = GetPipelinedFrameStateForThread();
+	const FPipelinedFrameState& FrameState = LockedPipelineState.GetFrameState();
 	if (FrameState.Views.Num() < 2)
 	{
 		return 0.064f;
@@ -646,7 +650,8 @@ bool FOpenXRHMD::GetIsTracked(int32 DeviceId)
 {
 	// This function is called from both the game and rendering thread and each thread maintains separate pose
 	// snapshots to prevent inconsistent poses (tearing) on the same frame.
-	const FPipelinedFrameState& PipelineState = GetPipelinedFrameStateForThread();
+	FPipelinedFrameStateAccessorReadOnly LockedPipelineState = (const_cast<const FOpenXRHMD*>(this))->GetPipelinedFrameStateForThread();
+	const FPipelinedFrameState& PipelineState = LockedPipelineState.GetFrameState();
 
 	if (!PipelineState.DeviceLocations.IsValidIndex(DeviceId))
 	{
@@ -665,7 +670,8 @@ bool FOpenXRHMD::GetCurrentPose(int32 DeviceId, FQuat& CurrentOrientation, FVect
 
 	// This function is called from both the game and rendering thread and each thread maintains separate pose
 	// snapshots to prevent inconsistent poses (tearing) on the same frame.
-	const FPipelinedFrameState& PipelineState = GetPipelinedFrameStateForThread();
+	FPipelinedFrameStateAccessorReadOnly LockedPipelineState = (const_cast<const FOpenXRHMD*>(this))->GetPipelinedFrameStateForThread();
+	const FPipelinedFrameState& PipelineState = LockedPipelineState.GetFrameState();
 
 	if (!PipelineState.DeviceLocations.IsValidIndex(DeviceId))
 	{
@@ -686,7 +692,8 @@ bool FOpenXRHMD::GetCurrentPose(int32 DeviceId, FQuat& CurrentOrientation, FVect
 
 bool FOpenXRHMD::GetPoseForTime(int32 DeviceId, FTimespan Timespan, bool& OutTimeWasUsed, FQuat& Orientation, FVector& Position, bool& bProvidedLinearVelocity, FVector& LinearVelocity, bool& bProvidedAngularVelocity, FVector& AngularVelocityAsAxisAndLength, bool& bProvidedLinearAcceleration, FVector& LinearAcceleration, float InWorldToMetersScale)
 {
-	const FPipelinedFrameState& PipelineState = GetPipelinedFrameStateForThread();
+	FPipelinedFrameStateAccessorReadOnly LockedPipelineState = (const_cast<const FOpenXRHMD*>(this))->GetPipelinedFrameStateForThread();
+	const FPipelinedFrameState& PipelineState = LockedPipelineState.GetFrameState();
 
 	FReadScopeLock DeviceLock(DeviceMutex);
 	if (!DeviceSpaces.IsValidIndex(DeviceId))
@@ -1055,7 +1062,8 @@ FIntPoint GeneratePixelDensitySize(const XrViewConfigurationView& Config, const
 
 void FOpenXRHMD::AdjustViewRect(int32 ViewIndex, int32& X, int32& Y, uint32& SizeX, uint32& SizeY) const
 {
-	const FPipelinedFrameState& PipelineState = GetPipelinedFrameStateForThread();
+	FPipelinedFrameStateAccessorReadOnly LockedPipelineState = GetPipelinedFrameStateForThread();
+	const FPipelinedFrameState& PipelineState = LockedPipelineState.GetFrameState();
 	const XrViewConfigurationView& Config = PipelineState.ViewConfigs[ViewIndex];
 	FIntPoint ViewRectMin(EForceInit::ForceInitToZero);
 
@@ -1081,7 +1089,8 @@ void FOpenXRHMD::CalculateRenderTargetSize(const FViewport& Viewport, uint32& In
 {
 	check(IsInGameThread() || IsInRenderingThread());
 
-	const FPipelinedFrameState& PipelineState = GetPipelinedFrameStateForThread();
+	FPipelinedFrameStateAccessorReadOnly LockedPipelineState = (const_cast<const FOpenXRHMD*>(this))->GetPipelinedFrameStateForThread();
+	const FPipelinedFrameState& PipelineState = LockedPipelineState.GetFrameState();
 	const float PixelDensity = PipelineState.PixelDensity;
 
 	// TODO: Could we just call AdjustViewRect per view, or even for _only_ the last view?
@@ -1154,7 +1163,8 @@ uint32 FOpenXRHMD::GetLODViewIndex() const
 
 int32 FOpenXRHMD::GetDesiredNumberOfViews(bool bStereoRequested) const
 {
-	const FPipelinedFrameState& FrameState = GetPipelinedFrameStateForThread();
+	FPipelinedFrameStateAccessorReadOnly LockedPipelineState = (const_cast<const FOpenXRHMD*>(this))->GetPipelinedFrameStateForThread();
+	const FPipelinedFrameState& FrameState = LockedPipelineState.GetFrameState();
 
 	// FIXME: Monoscopic actually needs 2 views for quad vr
 	return bStereoRequested ? FrameState.ViewConfigs.Num() : 1;
@@ -1167,7 +1177,8 @@ bool FOpenXRHMD::GetRelativeEyePose(int32 InDeviceId, int32 InViewIndex, FQuat&
 		return false;
 	}
 
-	const FPipelinedFrameState& FrameState = GetPipelinedFrameStateForThread();
+	FPipelinedFrameStateAccessorReadOnly LockedPipelineState = (const_cast<const FOpenXRHMD*>(this))->GetPipelinedFrameStateForThread();
+	const FPipelinedFrameState& FrameState = LockedPipelineState.GetFrameState();
 
 	if (FrameState.ViewState.viewStateFlags & XR_VIEW_STATE_ORIENTATION_VALID_BIT &&
 		FrameState.ViewState.viewStateFlags & XR_VIEW_STATE_POSITION_VALID_BIT &&
@@ -1183,7 +1194,8 @@ bool FOpenXRHMD::GetRelativeEyePose(int32 InDeviceId, int32 InViewIndex, FQuat&
 
 FMatrix FOpenXRHMD::GetStereoProjectionMatrix(const int32 ViewIndex) const
 {
-	const FPipelinedFrameState& FrameState = GetPipelinedFrameStateForThread();
+	FPipelinedFrameStateAccessorReadOnly LockedPipelineState = GetPipelinedFrameStateForThread();
+	const FPipelinedFrameState& FrameState = LockedPipelineState.GetFrameState();
 
 	XrFovf Fov = {};
 	if (ViewIndex == eSSE_MONOSCOPIC)
@@ -1238,7 +1250,8 @@ void FOpenXRHMD::SetupViewFamily(FSceneViewFamily& InViewFamily)
 	InViewFamily.EngineShowFlags.HMDDistortion = false;
 	InViewFamily.EngineShowFlags.StereoRendering = IsStereoEnabled();
 
-	const FPipelinedFrameState& FrameState = GetPipelinedFrameStateForThread();
+	FPipelinedFrameStateAccessorReadOnly LockedPipelineState = (const_cast<const FOpenXRHMD*>(this))->GetPipelinedFrameStateForThread();
+	const FPipelinedFrameState& FrameState = LockedPipelineState.GetFrameState();
 	if (FrameState.Views.Num() > 2)
 	{
 		InViewFamily.EngineShowFlags.Vignette = 0;
@@ -1537,7 +1550,7 @@ TArray<XrEnvironmentBlendMode> FOpenXRHMD::RetrieveEnvironmentBlendModes() const
 	return BlendModes;
 }
 
-const FOpenXRHMD::FPipelinedFrameState& FOpenXRHMD::GetPipelinedFrameStateForThread() const
+FOpenXRHMD::FPipelinedFrameStateAccessorReadOnly FOpenXRHMD::GetPipelinedFrameStateForThread() const
 {
 	// Relying on implicit selection of the RHI struct is hazardous since the RHI thread isn't always present
 	check(!IsInRHIThread());
@@ -1547,28 +1560,28 @@ const FOpenXRHMD::FPipelinedFrameState& FOpenXRHMD::GetPipelinedFrameStateForThr
 	// render thread at this moment is modifying the state using the non-const method. Proper resolution is tracked in UE-212224.
 	if (IsInActualRenderingThread() || IsInParallelRenderingThread())
 	{
-		return PipelinedFrameStateRendering;
+		return FPipelinedFrameStateAccessorReadOnly(PipelinedFrameStateRendering, PipelinedFrameStateRenderingAccessGuard);
 	}
 	else
 	{
-		check(IsInGameThread());
-		return PipelinedFrameStateGame;
+		check(IsInGameThread() || IsInParallelGameThread());
+		return FPipelinedFrameStateAccessorReadOnly(PipelinedFrameStateGame, PipelinedFrameStateGameAccessGuard);
 	}
 }
 
-FOpenXRHMD::FPipelinedFrameState& FOpenXRHMD::GetPipelinedFrameStateForThread()
+FOpenXRHMD::FPipelinedFrameStateAccessorReadWrite FOpenXRHMD::GetPipelinedFrameStateForThread()
 {
 	// Relying on implicit selection of the RHI struct is hazardous since the RHI thread isn't always present
 	check(!IsInRHIThread());
 
-	if (IsInActualRenderingThread())
+	if (IsInActualRenderingThread() || IsInParallelRenderingThread())
 	{
-		return PipelinedFrameStateRendering;
+		return FPipelinedFrameStateAccessorReadWrite(PipelinedFrameStateRendering, PipelinedFrameStateRenderingAccessGuard);
 	}
 	else
 	{
-		check(IsInGameThread());
-		return PipelinedFrameStateGame;
+		check(IsInGameThread() || IsInParallelGameThread());
+		return FPipelinedFrameStateAccessorReadWrite(PipelinedFrameStateGame, PipelinedFrameStateGameAccessGuard);
 	}
 }
 
@@ -1576,7 +1589,8 @@ void FOpenXRHMD::UpdateDeviceLocations(bool bUpdateOpenXRExtensionPlugins)
 {
 	SCOPED_NAMED_EVENT(UpdateDeviceLocations, FColor::Red);
 
-	FPipelinedFrameState& PipelineState = GetPipelinedFrameStateForThread();
+	FPipelinedFrameStateAccessorReadWrite LockedPipelineState = GetPipelinedFrameStateForThread();
+	FPipelinedFrameState& PipelineState = LockedPipelineState.GetFrameState();
 
 	// Only update the device locations if the frame state has been predicted, which is dependent on WaitFrame success
 	// Also need a valid TrackingSpace
@@ -2248,13 +2262,15 @@ XrSpace FOpenXRHMD::GetTrackedDeviceSpace(const int32 DeviceId)
 
 XrTime FOpenXRHMD::GetDisplayTime() const
 {
-	const FPipelinedFrameState& PipelineState = GetPipelinedFrameStateForThread();
+	FPipelinedFrameStateAccessorReadOnly LockedPipelineState = GetPipelinedFrameStateForThread();
+	const FPipelinedFrameState& PipelineState = LockedPipelineState.GetFrameState();
 	return PipelineState.bXrFrameStateUpdated ? PipelineState.FrameState.predictedDisplayTime : 0;
 }
 
 XrSpace FOpenXRHMD::GetTrackingSpace() const
 {
-	const FPipelinedFrameState& PipelineState = GetPipelinedFrameStateForThread();
+	FPipelinedFrameStateAccessorReadOnly LockedPipelineState = GetPipelinedFrameStateForThread();
+	const FPipelinedFrameState& PipelineState = LockedPipelineState.GetFrameState();
 	if (PipelineState.TrackingSpace.IsValid())
 	{
 		return PipelineState.TrackingSpace->Handle;
@@ -3140,7 +3156,8 @@ void FOpenXRHMD::OnBeginSimulation_GameThread()
 		return;
 	}
 
-	FPipelinedFrameState& PipelineState = GetPipelinedFrameStateForThread();
+	FPipelinedFrameStateAccessorReadWrite LockedPipelineState = GetPipelinedFrameStateForThread();
+	FPipelinedFrameState& PipelineState = LockedPipelineState.GetFrameState();
 	PipelineState.bXrFrameStateUpdated = false;
 	PipelineState.FrameState = { XR_TYPE_FRAME_STATE };
 
@@ -3665,7 +3682,8 @@ bool FOpenXRHMD::HDRGetMetaDataForStereo(EDisplayOutputFormat& OutDisplayOutputF
 
 float FOpenXRHMD::GetPixelDenity() const
 {
-	const FPipelinedFrameState& PipelineState = GetPipelinedFrameStateForThread();
+	FPipelinedFrameStateAccessorReadOnly LockedPipelineState = GetPipelinedFrameStateForThread();
+	const FPipelinedFrameState& PipelineState = LockedPipelineState.GetFrameState();
 	return PipelineState.PixelDensity;
 }
 
@@ -3685,7 +3703,8 @@ void FOpenXRHMD::SetPixelDensity(const float NewDensity)
 
 FIntPoint FOpenXRHMD::GetIdealRenderTargetSize() const
 {
-	const FPipelinedFrameState& PipelineState = GetPipelinedFrameStateForThread();
+	FPipelinedFrameStateAccessorReadOnly LockedPipelineState = GetPipelinedFrameStateForThread();
+	const FPipelinedFrameState& PipelineState = LockedPipelineState.GetFrameState();
 
 	FIntPoint Size(EForceInit::ForceInitToZero);
 	for (int32 ViewIndex = 0; ViewIndex < PipelineState.ViewConfigs.Num(); ViewIndex++)
diff --git a/Engine/Plugins/Runtime/OpenXR/Source/OpenXRHMD/Private/OpenXRHMD.h b/Engine/Plugins/Runtime/OpenXR/Source/OpenXRHMD/Private/OpenXRHMD.h
--- a/Engine/Plugins/Runtime/OpenXR/Source/OpenXRHMD/Private/OpenXRHMD.h
+++ b/Engine/Plugins/Runtime/OpenXR/Source/OpenXRHMD/Private/OpenXRHMD.h
@@ -95,6 +95,62 @@ public:
 		bool bXrFrameStateUpdated = false;
 	};
 
+	struct FPipelinedFrameStateAccessorReadOnly
+	{
+		FPipelinedFrameStateAccessorReadOnly(FPipelinedFrameState const& InPipelinedFrameState, FRWLock const& InAccessGuard)
+			: PipelinedFrameState(InPipelinedFrameState), AccessGuard(const_cast<FRWLock&>(InAccessGuard))
+		{
+			AccessGuard.ReadLock();
+		}
+
+		// not virtual because we aren't going to delete via a base pointer
+		~FPipelinedFrameStateAccessorReadOnly()
+		{
+			AccessGuard.ReadUnlock();
+		}
+
+		FPipelinedFrameState const& GetFrameState() const
+		{
+			return PipelinedFrameState;
+		}
+
+	private:
+
+		/** The frame state we're guarding the access to. */
+		FPipelinedFrameState const& PipelinedFrameState;
+
+		/** Reference to the guarding lock. */
+		FRWLock& AccessGuard;
+	};
+
+	struct FPipelinedFrameStateAccessorReadWrite
+	{
+		FPipelinedFrameStateAccessorReadWrite(FPipelinedFrameState& InPipelinedFrameState, FRWLock& InAccessGuard)
+			: PipelinedFrameState(InPipelinedFrameState), AccessGuard(InAccessGuard)
+		{
+			AccessGuard.WriteLock();
+		}
+
+		// not virtual because we aren't going to delete via a base pointer
+		~FPipelinedFrameStateAccessorReadWrite()
+		{
+			AccessGuard.WriteUnlock();
+		}
+
+		FPipelinedFrameState& GetFrameState()
+		{
+			return PipelinedFrameState;
+		}
+
+	private:
+
+		/** The frame state we're guarding the access to. */
+		FPipelinedFrameState& PipelinedFrameState;
+
+		/** Reference to the guarding lock. */
+		FRWLock& AccessGuard;
+	};
+
 	struct FEmulatedLayerState
 	{
 		// These layers are used as a target to composite all the emulated face locked layers into
@@ -275,8 +331,8 @@ protected:
 	void BuildOcclusionMeshes();
 	bool BuildOcclusionMesh(XrVisibilityMaskTypeKHR Type, int View, FHMDViewMesh& Mesh);
 
-	const FPipelinedFrameState& GetPipelinedFrameStateForThread() const;
-	FPipelinedFrameState& GetPipelinedFrameStateForThread();
+	FPipelinedFrameStateAccessorReadOnly GetPipelinedFrameStateForThread() const;
+	FPipelinedFrameStateAccessorReadWrite GetPipelinedFrameStateForThread();
 
 	void UpdateDeviceLocations(bool bUpdateOpenXRExtensionPlugins);
 	void EnumerateViews(FPipelinedFrameState& PipelineState);
@@ -479,6 +535,12 @@ private:
 	FPipelinedFrameState	PipelinedFrameStateRendering;
 	FPipelinedFrameState	PipelinedFrameStateRHI;
 
+	/** Arbitrates access to PipelinedFrameStateGame from code on task threads */
+	FRWLock					PipelinedFrameStateGameAccessGuard;
+	/** Arbitrates access to PipelinedFrameStateRendering from code on task threads */
+	FRWLock					PipelinedFrameStateRenderingAccessGuard;
+	// we do not expose PipelineFrameStateRHI in GetPipelinedFrameStateForThread() as of now, so no lock is provided for that one.
+
 	FPipelinedLayerState	PipelinedLayerStateRendering;
 	FPipelinedLayerState	PipelinedLayerStateRHI;
 
