UE-214973 IOpenXRExtensionPlugin::UpdateCompositionLayers - headers array should not be of pointer-to-const
-Removed const so that extension plugins can add new chain structs via the next ptr.
-UE-214756 Fixed passing pointer to openxr chain struct that goes out of scope.
-Removed a couple 5.1 deprecated functions from IOpenXRExtensionPlugin.h

diff --git a/Engine/Plugins/Runtime/OpenXR/Source/OpenXRHMD/Private/OpenXRHMD.cpp b/Engine/Plugins/Runtime/OpenXR/Source/OpenXRHMD/Private/OpenXRHMD.cpp
--- a/Engine/Plugins/Runtime/OpenXR/Source/OpenXRHMD/Private/OpenXRHMD.cpp
+++ b/Engine/Plugins/Runtime/OpenXR/Source/OpenXRHMD/Private/OpenXRHMD.cpp
@@ -3470,9 +3470,9 @@ void FOpenXRHMD::OnBeginRendering_RHIThread(IRHICommandContext& RHICmdContext, c
 	PipelinedFrameStateRHI = InFrameState;
 
 	void* Next = nullptr;
+	XrRHIContextEPIC RHIContextEPIC = { (XrStructureType)XR_TYPE_RHI_CONTEXT_EPIC };
 	if (RuntimeRequiresRHIContext())
 	{
-		XrRHIContextEPIC RHIContextEPIC = { (XrStructureType)XR_TYPE_RHI_CONTEXT_EPIC };
 		RHIContextEPIC.RHIContext = &RHICmdContext;
 		RHIContextEPIC.next = Next;
 		Next = &RHIContextEPIC;
@@ -3574,7 +3574,7 @@ void FOpenXRHMD::OnFinishRendering_RHIThread(IRHICommandContext& RHICmdContext)
 	FReadScopeLock Lock(SessionHandleMutex);
 	if (bIsRunning)
 	{
-		TArray<const XrCompositionLayerBaseHeader*> Headers;
+		TArray<XrCompositionLayerBaseHeader*> Headers;
 		XrCompositionLayerProjection Layer = {};
 		XrCompositionLayerAlphaBlendFB LayerAlphaBlend = { XR_TYPE_COMPOSITION_LAYER_ALPHA_BLEND_FB };
 		XrCompositionLayerColorScaleBiasKHR ColorScaleBias = { XR_TYPE_COMPOSITION_LAYER_COLOR_SCALE_BIAS_KHR };
@@ -3586,7 +3586,7 @@ void FOpenXRHMD::OnFinishRendering_RHIThread(IRHICommandContext& RHICmdContext)
 			Layer.space = PipelinedFrameStateRHI.TrackingSpace->Handle;
 			Layer.viewCount = PipelinedLayerStateRHI.ProjectionLayers.Num();
 			Layer.views = PipelinedLayerStateRHI.ProjectionLayers.GetData();
-			Headers.Add(reinterpret_cast<const XrCompositionLayerBaseHeader*>(&Layer));
+			Headers.Add(reinterpret_cast<XrCompositionLayerBaseHeader*>(&Layer));
 
 			if(IsExtensionEnabled(XR_FB_COMPOSITION_LAYER_ALPHA_BLEND_EXTENSION_NAME) &&
 				bOpenXRInvertAlphaCvarCachedValue)
@@ -3628,15 +3628,15 @@ void FOpenXRHMD::OnFinishRendering_RHIThread(IRHICommandContext& RHICmdContext)
 			}
 			CompositedLayer.viewCount = PipelinedLayerStateRHI.EmulatedLayerState.CompositedProjectionLayers.Num();
 			CompositedLayer.views = PipelinedLayerStateRHI.EmulatedLayerState.CompositedProjectionLayers.GetData();
-			Headers.Add(reinterpret_cast<const XrCompositionLayerBaseHeader*>(&CompositedLayer));
+			Headers.Add(reinterpret_cast<XrCompositionLayerBaseHeader*>(&CompositedLayer));
 		}
 
 		AddLayersToHeaders(Headers);
 
 		void* Next = nullptr;
+		XrRHIContextEPIC RHIContextEPIC = { (XrStructureType)XR_TYPE_RHI_CONTEXT_EPIC };
 		if (RuntimeRequiresRHIContext())
 		{
-			XrRHIContextEPIC RHIContextEPIC = { (XrStructureType)XR_TYPE_RHI_CONTEXT_EPIC };
 			RHIContextEPIC.RHIContext = &RHICmdContext;
 			RHIContextEPIC.next = Next;
 			Next = &RHIContextEPIC;
@@ -3674,11 +3674,11 @@ void FOpenXRHMD::OnFinishRendering_RHIThread(IRHICommandContext& RHICmdContext)
 	bIsRendering = false;
 }
 
-void FOpenXRHMD::AddLayersToHeaders(TArray<const XrCompositionLayerBaseHeader*>& Headers)
+void FOpenXRHMD::AddLayersToHeaders(TArray<XrCompositionLayerBaseHeader*>& Headers)
 {
-	for (const FXrCompositionLayerUnion& Layer : PipelinedLayerStateRHI.NativeOverlays)
+	for (FXrCompositionLayerUnion& Layer : PipelinedLayerStateRHI.NativeOverlays)
 	{
-		Headers.Add(reinterpret_cast<const XrCompositionLayerBaseHeader*>(&Layer.Header));
+		Headers.Add(reinterpret_cast<XrCompositionLayerBaseHeader*>(&Layer.Header));
 	}
 
 	for (IOpenXRExtensionPlugin* Module : ExtensionPlugins)
diff --git a/Engine/Plugins/Runtime/OpenXR/Source/OpenXRHMD/Private/OpenXRHMD.h b/Engine/Plugins/Runtime/OpenXR/Source/OpenXRHMD/Private/OpenXRHMD.h
--- a/Engine/Plugins/Runtime/OpenXR/Source/OpenXRHMD/Private/OpenXRHMD.h
+++ b/Engine/Plugins/Runtime/OpenXR/Source/OpenXRHMD/Private/OpenXRHMD.h
@@ -489,7 +489,7 @@ private:
 	
 	void UpdateLayerSwapchainTexture(const FOpenXRLayer& Layer, FRHICommandListImmediate& RHICmdList);
 	void ConfigureLayerSwapchain(FOpenXRLayer& Layer, TArray<FOpenXRLayer>& BackupLayers);
-	void AddLayersToHeaders(TArray<const XrCompositionLayerBaseHeader*>& Headers);
+	void AddLayersToHeaders(TArray<XrCompositionLayerBaseHeader*>& Headers);
 
 	bool					bStereoEnabled;
 	TAtomic<bool>			bIsRunning;
diff --git a/Engine/Plugins/Runtime/OpenXR/Source/OpenXRHMD/Public/IOpenXRExtensionPlugin.h b/Engine/Plugins/Runtime/OpenXR/Source/OpenXRHMD/Public/IOpenXRExtensionPlugin.h
--- a/Engine/Plugins/Runtime/OpenXR/Source/OpenXRHMD/Public/IOpenXRExtensionPlugin.h
+++ b/Engine/Plugins/Runtime/OpenXR/Source/OpenXRHMD/Public/IOpenXRExtensionPlugin.h
@@ -255,20 +255,6 @@ public:
 		return false;
 	}
 
-	/**
-	* Add any actions provided by the plugin to Actions.
-	* This allows a plugin to 'hard code' an action so that the plugin can use it.
-	*/
-	UE_DEPRECATED(5.1, "Use Enhanced Input through IMotionController::SetPlayerMappableInputConfig instead.")
-	virtual void AddActions(XrInstance Instance, TFunction<XrAction(XrActionType InActionType, const FName& InName, const TArray<XrPath>& InSubactionPaths)> AddAction)
-	{
-	}
-
-	UE_DEPRECATED(5.1, "Functionality moved to AttachActionSets().")
-	virtual void AddActionSets(TArray<XrActiveActionSet>& OutActionSets)
-	{
-	}
-
 	/**
 	* Add any action sets provided by the plugin to be attached as active to the session
 	* This allows a plugin to manage a custom actionset that will be active in xrSyncActions
@@ -397,10 +383,18 @@ public:
 	{
 	}
 
-	// FOpenXRHMD::OnFinishRendering_RHIThread
+	UE_DEPRECATED(5.5, "Please replace with the version that takes an array of non-const XrCompositionLayerBaseHeader*, which allows chain structs to be added via the next pointer.")
 	virtual void UpdateCompositionLayers(XrSession InSession, TArray<const XrCompositionLayerBaseHeader*>& Headers)
 	{
 	}
+	
+	// FOpenXRHMD::OnFinishRendering_RHIThread
+	virtual void UpdateCompositionLayers(XrSession InSession, TArray<XrCompositionLayerBaseHeader*>& Headers)
+	{
+		PRAGMA_DISABLE_DEPRECATION_WARNINGS
+		UpdateCompositionLayers(InSession, reinterpret_cast<TArray<const XrCompositionLayerBaseHeader*>&>(Headers));
+		PRAGMA_ENABLE_DEPRECATION_WARNINGS
+	}
 
 	virtual const void* OnEndProjectionLayer(XrSession InSession, int32 InLayerIndex, const void* InNext, XrCompositionLayerFlags& OutFlags)
 	{
diff --git a/Engine/Plugins/Runtime/OpenXR/Source/OpenXRInput/Private/OpenXRInput.cpp b/Engine/Plugins/Runtime/OpenXR/Source/OpenXRInput/Private/OpenXRInput.cpp
--- a/Engine/Plugins/Runtime/OpenXR/Source/OpenXRInput/Private/OpenXRInput.cpp
+++ b/Engine/Plugins/Runtime/OpenXR/Source/OpenXRInput/Private/OpenXRInput.cpp
@@ -525,10 +525,6 @@ bool FOpenXRInputPlugin::FOpenXRInput::BuildActions(XrSession Session)
 	for (IOpenXRExtensionPlugin* Plugin : OpenXRHMD->GetExtensionPlugins())
 	{
 		TArray<XrActiveActionSet> PluginAttachArray_Deprecated;
-		PRAGMA_DISABLE_DEPRECATION_WARNINGS;
-		// TODO?: Log deprecation warning at runtime, since overridden deprecated interface methods don't warn at compile time?
-		Plugin->AddActionSets(PluginAttachArray_Deprecated);
-		PRAGMA_ENABLE_DEPRECATION_WARNINGS;
 		for (const XrActiveActionSet& ActiveSet : PluginAttachArray_Deprecated)
 		{
 			AttachSet.Add(ActiveSet.actionSet);
@@ -587,22 +583,6 @@ void FOpenXRInputPlugin::FOpenXRInput::BuildLegacyActions(TMap<FString, FInterac
 		}
 	}
 
-	// Query extension plugins for actions
-	for (IOpenXRExtensionPlugin* Plugin : OpenXRHMD->GetExtensionPlugins())
-	{
-		PRAGMA_DISABLE_DEPRECATION_WARNINGS;
-		Plugin->AddActions(Instance,
-			[this, &ActionSet](XrActionType InActionType, const FName& InName, const TArray<XrPath>& InSubactionPaths)
-			{
-				// TODO?: Log deprecation warning at runtime, since overridden deprecated interface methods don't warn at compile time?
-				FOpenXRAction Action(ActionSet.Handle, InActionType, InName, InName.ToString(), InSubactionPaths);
-				LegacyActions.Add(Action);
-				return Action.Handle;
-			}
-		);
-		PRAGMA_ENABLE_DEPRECATION_WARNINGS;
-	}
-
 	ActionSets.Emplace(MoveTemp(ActionSet));
 }
 
