Static analysis fix.
-A few more lines needed to be removed with the deprecated function removal.

diff --git a/Engine/Plugins/Runtime/OpenXR/Source/OpenXRInput/Private/OpenXRInput.cpp b/Engine/Plugins/Runtime/OpenXR/Source/OpenXRInput/Private/OpenXRInput.cpp
--- a/Engine/Plugins/Runtime/OpenXR/Source/OpenXRInput/Private/OpenXRInput.cpp
+++ b/Engine/Plugins/Runtime/OpenXR/Source/OpenXRInput/Private/OpenXRInput.cpp
@@ -524,12 +524,6 @@ bool FOpenXRInputPlugin::FOpenXRInput::BuildActions(XrSession Session)
 	// Bind plugin action sets exposed through a query
 	for (IOpenXRExtensionPlugin* Plugin : OpenXRHMD->GetExtensionPlugins())
 	{
-		TArray<XrActiveActionSet> PluginAttachArray_Deprecated;
-		for (const XrActiveActionSet& ActiveSet : PluginAttachArray_Deprecated)
-		{
-			AttachSet.Add(ActiveSet.actionSet);
-		}
-
 		TSet<XrActionSet> PluginAttachSet;
 		Plugin->AttachActionSets(PluginAttachSet);
 		AttachSet.Append(PluginAttachSet);
