Input Actions with the same description silently breaks OpenXR input
-Detecting duplicate Action or ActionSet descriptions before we create them, issuing a warning, and then replacing them with the Action or ActionSet's FName so that the inputs will function.

diff --git a/Engine/Plugins/Runtime/OpenXR/Source/OpenXRInput/Private/OpenXRInput.cpp b/Engine/Plugins/Runtime/OpenXR/Source/OpenXRInput/Private/OpenXRInput.cpp
--- a/Engine/Plugins/Runtime/OpenXR/Source/OpenXRInput/Private/OpenXRInput.cpp
+++ b/Engine/Plugins/Runtime/OpenXR/Source/OpenXRInput/Private/OpenXRInput.cpp
@@ -588,10 +588,36 @@ void FOpenXRInputPlugin::FOpenXRInput::BuildEnhancedActions(TMap<FString, FInter
 		return;
 	}
 
+	// OpenXR does not allow duplicated localized names for mapping contexts.  In order to have good warnings and allow input to function even with bad localization setup we will detect duplicates ourselves.
+	TMap<FString, FName> MappingContextDescriptionMap;
+
 	PRAGMA_DISABLE_DEPRECATION_WARNINGS
 	for (const auto& MappingContext : InputMappingContextToPriorityMap)
 	{
-		FOpenXRActionSet ActionSet(Instance, MappingContext.Key->GetFName(), MappingContext.Key->ContextDescription.ToString(), MappingContext.Value, MappingContext.Key.Get());
+		FName NewMappingContextName = MappingContext.Key->GetFName();
+		FString NewMappingContextDescription = MappingContext.Key->ContextDescription.ToString();
+
+		// Handle any description string duplicates
+		{
+			FName* ExistingMappingContextName = MappingContextDescriptionMap.Find(NewMappingContextDescription);
+			if (ExistingMappingContextName != nullptr)
+			{
+				UE_LOG(LogHMD, Warning, TEXT("Input Mapping Context %s has a Description, \"%s\", which exactly matches the Description already used by Input Mapping Context %s.  Identical localized descriptions are not allowed by OpenXR.  The FName of this mapping context %s, which is unique, will replace the duplicated localized string so the Input Mapping Context functions but it will not localize correctly.")
+					, *NewMappingContextName.ToString(), *NewMappingContextDescription, *ExistingMappingContextName->ToString(), *NewMappingContextName.ToString());
+				NewMappingContextDescription = NewMappingContextName.ToString();
+			}
+			else
+			{
+				MappingContextDescriptionMap.Add(NewMappingContextDescription, NewMappingContextName);
+			}
+		}
+
+		FOpenXRActionSet ActionSet(Instance, NewMappingContextName, NewMappingContextDescription, MappingContext.Value, MappingContext.Key.Get());
+
+
+		// OpenXR does not allow duplicated localized names for actions.  In order to have good warnings allow input to function even with bad localization setup we will detect duplicates ourselves.
+		TMap<FString, FName> ActionDescriptionMap;
+
 		TMap<FName, int32> ActionMap;
 
 		for (const FEnhancedActionKeyMapping& Mapping : MappingContext.Key->GetMappings())
@@ -602,12 +628,28 @@ void FOpenXRInputPlugin::FOpenXRInput::BuildEnhancedActions(TMap<FString, FInter
 			}
 
 			// Try to find an existing action within the current action set
-			FName ActionName = Mapping.Action->GetFName();
-			int32& ActionIndex = ActionMap.FindOrAdd(ActionName, INDEX_NONE);
+			FName NewActionName = Mapping.Action->GetFName();
+			int32& ActionIndex = ActionMap.FindOrAdd(NewActionName, INDEX_NONE);
 			if (ActionIndex == INDEX_NONE)
 			{
 				// No action found, create a new one
-				FString LocalizedName = Mapping.Action->ActionDescription.ToString();
+				FString NewActionDescription = Mapping.Action->ActionDescription.ToString();
+
+				// Handle any description string duplicates
+				{
+					FName* ExistingActionName = ActionDescriptionMap.Find(NewActionDescription);
+					if (ExistingActionName != nullptr)
+					{
+						UE_LOG(LogHMD, Warning, TEXT("Input Action %s has a Description, \"%s\", which exactly matches the Description already used by Action %s.  Identical localized descriptions are not allowed by OpenXR.  The FName of this action %s, which is unique, will replace the localized string so that the input Action functions but it will not localize correctly.")
+							, *NewActionName.ToString(), *NewActionDescription, *ExistingActionName->ToString(), *NewActionName.ToString());
+						NewActionDescription = NewActionName.ToString();
+					}
+					else
+					{
+						ActionDescriptionMap.Add(NewActionDescription, NewActionName);
+					}
+				}
+
 				XrActionType ActionType = ToActionType(Mapping.Action->ValueType);
 				if (!ActionType)
 				{
@@ -615,7 +657,7 @@ void FOpenXRInputPlugin::FOpenXRInput::BuildEnhancedActions(TMap<FString, FInter
 				}
 
 				// Create the action and write the index to the reference in the actions map
-				ActionIndex = EnhancedActions.Emplace(ActionSet.Handle, ActionType, ActionName, LocalizedName, SubactionPaths, Mapping.Action);
+				ActionIndex = EnhancedActions.Emplace(ActionSet.Handle, ActionType, NewActionName, NewActionDescription, SubactionPaths, Mapping.Action);
 			}
 
 			SuggestBindingForKey(Profiles, EnhancedActions[ActionIndex], Mapping.Key, Mapping.Modifiers, Mapping.Triggers);
@@ -1257,8 +1299,6 @@ bool FOpenXRInputPlugin::FOpenXRInput::GetControllerOrientationAndPositionForTim
 
 	if (GetInfo.action == XR_NULL_HANDLE)
 	{
-		UE_LOG(LogHMD, Warning, TEXT("GetControllerOrientationAndPositionForTime called with motion source %s which is unknown.  Cannot get pose."), *MotionSource.ToString());
-
 		return false;
 	}
 
