Check to see if we're in an XR Tracking Only session and update tracking devices if we are

diff --git a/Engine/Plugins/Runtime/OpenXR/Source/OpenXRHMD/Private/OpenXRHMD.cpp b/Engine/Plugins/Runtime/OpenXR/Source/OpenXRHMD/Private/OpenXRHMD.cpp
--- a/Engine/Plugins/Runtime/OpenXR/Source/OpenXRHMD/Private/OpenXRHMD.cpp
+++ b/Engine/Plugins/Runtime/OpenXR/Source/OpenXRHMD/Private/OpenXRHMD.cpp
@@ -3321,7 +3321,7 @@ void FOpenXRHMD::OnBeginSimulation_GameThread()
 {
 	FReadScopeLock Lock(SessionHandleMutex);
 
-	if (!bShouldWait || !RenderBridge)
+	if (!bShouldWait || (!RenderBridge && !bIsTrackingOnlySession))
 	{
 		return;
 	}
