[OpenXR]One extension plugin adds multiple interaction profiles
-Adds a hook to IOpenXRExtensionPlugin to add multiple interaction profiles.  Previously only one could be added.
-Deprecated the single interaction profile version, and replaced all uses of it in the engine.

diff --git a/Engine/Plugins/Runtime/HPMotionController/Source/HPMotionController/Private/HPMotionController.cpp b/Engine/Plugins/Runtime/HPMotionController/Source/HPMotionController/Private/HPMotionController.cpp
--- a/Engine/Plugins/Runtime/HPMotionController/Source/HPMotionController/Private/HPMotionController.cpp
+++ b/Engine/Plugins/Runtime/HPMotionController/Source/HPMotionController/Private/HPMotionController.cpp
@@ -132,11 +132,11 @@ void FHPMotionControllerModule::PostCreateInstance(XrInstance InInstance)
 	ControllerModels.Add(Path, FOpenXRAssetDirectory::HPMixedRealityRight);
 }
 
-bool FHPMotionControllerModule::GetInteractionProfile(XrInstance InInstance, FString& OutKeyPrefix, XrPath& OutPath, bool& OutHasHaptics)
+bool FHPMotionControllerModule::GetInteractionProfiles(XrInstance InInstance, TArray<FString>& OutKeyPrefixes, TArray<XrPath>& OutPaths, TArray<bool>& OutHasHaptics)
 {
-	OutKeyPrefix = "HPMixedRealityController";
-	OutPath = InteractionProfile;
-	OutHasHaptics = true;
+	OutKeyPrefixes.Add("HPMixedRealityController");
+	OutPaths.Add(InteractionProfile);
+	OutHasHaptics.Add(true);
 
 	return true;
 }
diff --git a/Engine/Plugins/Runtime/HPMotionController/Source/HPMotionController/Private/HPMotionController.h b/Engine/Plugins/Runtime/HPMotionController/Source/HPMotionController/Private/HPMotionController.h
--- a/Engine/Plugins/Runtime/HPMotionController/Source/HPMotionController/Private/HPMotionController.h
+++ b/Engine/Plugins/Runtime/HPMotionController/Source/HPMotionController/Private/HPMotionController.h
@@ -22,7 +22,7 @@ public:
 
 	bool GetRequiredExtensions(TArray<const ANSICHAR*>& OutExtensions) override;
 	void PostCreateInstance(XrInstance InInstance) override;
-	bool GetInteractionProfile(XrInstance InInstance, FString& OutKeyPrefix, XrPath& OutPath, bool& OutHasHaptics) override;
+	bool GetInteractionProfiles(XrInstance InInstance, TArray<FString>& OutKeyPrefixes, TArray<XrPath>& OutPaths, TArray<bool>& OutHasHaptics) override;
 	bool GetControllerModel(XrInstance InInstance, XrPath InInteractionProfile, XrPath InDevicePath, FSoftObjectPath& OutPath) override;
 	void GetControllerModelsForCooking(TArray<FSoftObjectPath>& OutPaths) override;
 
diff --git a/Engine/Plugins/Runtime/OpenXR/Source/OpenXRHMD/Public/IOpenXRExtensionPlugin.h b/Engine/Plugins/Runtime/OpenXR/Source/OpenXRHMD/Public/IOpenXRExtensionPlugin.h
--- a/Engine/Plugins/Runtime/OpenXR/Source/OpenXRHMD/Public/IOpenXRExtensionPlugin.h
+++ b/Engine/Plugins/Runtime/OpenXR/Source/OpenXRHMD/Public/IOpenXRExtensionPlugin.h
@@ -209,11 +209,20 @@ public:
 	/**
 	* Set the output parameters to add an interaction profile to OpenXR Input
 	*/
+	UE_DEPRECATED(5.5, "Deprecated in favor of the same-name function which allows the addition of multiple interaction profiles.")
 	virtual bool GetInteractionProfile(XrInstance InInstance, FString& OutKeyPrefix, XrPath& OutPath, bool& OutHasHaptics)
 	{
 		return false;
 	}
 
+	/**
+	* Set the output parameters to add multiple interaction profiles to OpenXR Input
+	*/
+	virtual bool GetInteractionProfiles(XrInstance InInstance, TArray<FString>& OutKeyPrefixes, TArray<XrPath>& OutPaths, TArray<bool>& OutHasHaptics)
+	{
+		return false;
+	}
+
 	/**
 	 * Set the output parameter to add suggested bindings to the given interaction profile.
 	 * This function gets called once for each interaction profile.
diff --git a/Engine/Plugins/Runtime/OpenXR/Source/OpenXRInput/Private/OpenXRInput.cpp b/Engine/Plugins/Runtime/OpenXR/Source/OpenXRInput/Private/OpenXRInput.cpp
--- a/Engine/Plugins/Runtime/OpenXR/Source/OpenXRInput/Private/OpenXRInput.cpp
+++ b/Engine/Plugins/Runtime/OpenXR/Source/OpenXRInput/Private/OpenXRInput.cpp
@@ -391,17 +391,39 @@ bool FOpenXRInputPlugin::FOpenXRInput::BuildActions(XrSession Session)
 	Profiles.Add("ValveIndex", FInteractionProfile(FOpenXRPath("/interaction_profiles/valve/index_controller"), true));
 
 	// Query extension plugins for interaction profiles
-	for (IOpenXRExtensionPlugin* Plugin : OpenXRHMD->GetExtensionPlugins())
 	{
-		FString KeyPrefix;
-		XrPath Path = XR_NULL_PATH;
-		bool HasHaptics = false;
-		if (Plugin->GetInteractionProfile(Instance, KeyPrefix, Path, HasHaptics) && Path != XR_NULL_PATH)
+		TArray<FString> KeyPrefixes;
+		TArray<XrPath> Paths;
+		TArray<bool> Haptics;
+		for (IOpenXRExtensionPlugin* Plugin : OpenXRHMD->GetExtensionPlugins())
 		{
-			Profiles.Add(KeyPrefix, FInteractionProfile(Path, HasHaptics));
+			FString KeyPrefix;
+			XrPath Path = XR_NULL_PATH;
+			bool HasHaptics = false;
+			PRAGMA_DISABLE_DEPRECATION_WARNINGS
+			if (Plugin->GetInteractionProfile(Instance, KeyPrefix, Path, HasHaptics) && Path != XR_NULL_PATH)
+			{
+				Profiles.Add(KeyPrefix, FInteractionProfile(Path, HasHaptics));
+			}
+			PRAGMA_ENABLE_DEPRECATION_WARNINGS
+	
+			KeyPrefixes.Reset();
+			Paths.Reset();
+			Haptics.Reset();
+			if (Plugin->GetInteractionProfiles(Instance, KeyPrefixes, Paths, Haptics) && Paths.Num() != 0)
+			{
+				for (int i = 0; i < Paths.Num(); i++)
+				{
+					if (Paths[i] != XR_NULL_PATH)
+					{
+						Profiles.Add(KeyPrefixes[i], FInteractionProfile(Paths[i], Haptics[i]));
+					}
+				}
+			}
 		}
 	}
 
+
 	PRAGMA_DISABLE_DEPRECATION_WARNINGS
 	// Attempt to load the default input config from the OpenXR input settings.
 	const UEnhancedInputDeveloperSettings* InputSettings = GetDefault<UEnhancedInputDeveloperSettings>();
diff --git a/Engine/Plugins/Runtime/OpenXREyeTracker/Source/OpenXREyeTracker/Private/OpenXREyeTracker.cpp b/Engine/Plugins/Runtime/OpenXREyeTracker/Source/OpenXREyeTracker/Private/OpenXREyeTracker.cpp
--- a/Engine/Plugins/Runtime/OpenXREyeTracker/Source/OpenXREyeTracker/Private/OpenXREyeTracker.cpp
+++ b/Engine/Plugins/Runtime/OpenXREyeTracker/Source/OpenXREyeTracker/Private/OpenXREyeTracker.cpp
@@ -53,11 +53,14 @@ void FOpenXREyeTracker::PostCreateInstance(XrInstance InInstance)
 	Instance = InInstance;
 }
 
-bool FOpenXREyeTracker::GetInteractionProfile(XrInstance InInstance, FString& OutKeyPrefix, XrPath& OutPath, bool& OutHasHaptics)
+bool FOpenXREyeTracker::GetInteractionProfiles(XrInstance InInstance, TArray<FString>& OutKeyPrefixes, TArray<XrPath>& OutPaths, TArray<bool>& OutHasHaptics)
 {
-	OutKeyPrefix = "EyeTracker";
-	OutHasHaptics = false;
-	return xrStringToPath(InInstance, "/interaction_profiles/ext/eye_gaze_interaction", &OutPath) == XR_SUCCESS;
+	OutKeyPrefixes.Add("EyeTracker");
+	OutHasHaptics.Add(false);
+	XrPath OutPath;
+	XrResult Result = xrStringToPath(InInstance, "/interaction_profiles/ext/eye_gaze_interaction", &OutPath);
+	OutPaths.Add(OutPath);
+	return Result == XR_SUCCESS;
 }
 
 void FOpenXREyeTracker::AttachActionSets(TSet<XrActionSet>& OutActionSets)
diff --git a/Engine/Plugins/Runtime/OpenXREyeTracker/Source/OpenXREyeTracker/Private/OpenXREyeTracker.h b/Engine/Plugins/Runtime/OpenXREyeTracker/Source/OpenXREyeTracker/Private/OpenXREyeTracker.h
--- a/Engine/Plugins/Runtime/OpenXREyeTracker/Source/OpenXREyeTracker/Private/OpenXREyeTracker.h
+++ b/Engine/Plugins/Runtime/OpenXREyeTracker/Source/OpenXREyeTracker/Private/OpenXREyeTracker.h
@@ -36,7 +36,7 @@ public:
 	}
 	virtual bool GetRequiredExtensions(TArray<const ANSICHAR*>& OutExtensions) override;
 	virtual void PostCreateInstance(XrInstance InInstance) override;
-	virtual bool GetInteractionProfile(XrInstance InInstance, FString& OutKeyPrefix, XrPath& OutPath, bool& OutHasHaptics) override;
+	virtual bool GetInteractionProfiles(XrInstance InInstance, TArray<FString>& OutKeyPrefixes, TArray<XrPath>& OutPaths, TArray<bool>& OutHasHaptics) override;
 	virtual void AttachActionSets(TSet<XrActionSet>& OutActionSets) override;
 	virtual const void* OnBeginSession(XrSession InSession, const void* InNext) override;
 	virtual void OnDestroySession(XrSession InSession) override;
diff --git a/Engine/Plugins/Runtime/OpenXRMsftHandInteraction/Source/OpenXRMsftHandInteraction/Private/OpenXRMsftHandInteraction.cpp b/Engine/Plugins/Runtime/OpenXRMsftHandInteraction/Source/OpenXRMsftHandInteraction/Private/OpenXRMsftHandInteraction.cpp
--- a/Engine/Plugins/Runtime/OpenXRMsftHandInteraction/Source/OpenXRMsftHandInteraction/Private/OpenXRMsftHandInteraction.cpp
+++ b/Engine/Plugins/Runtime/OpenXRMsftHandInteraction/Source/OpenXRMsftHandInteraction/Private/OpenXRMsftHandInteraction.cpp
@@ -37,11 +37,14 @@ bool FOpenXRMsftHandInteraction::GetRequiredExtensions(TArray<const ANSICHAR*>&
 	return true;
 }
 
-bool FOpenXRMsftHandInteraction::GetInteractionProfile(XrInstance InInstance, FString& OutKeyPrefix, XrPath& OutPath, bool& OutHasHaptics)
+bool FOpenXRMsftHandInteraction::GetInteractionProfiles(XrInstance InInstance, TArray<FString>& OutKeyPrefixes, TArray<XrPath>& OutPaths, TArray<bool>& OutHasHaptics)
 {
-	OutKeyPrefix = "OpenXRMsftHandInteraction";
-	OutHasHaptics = false;
-	return xrStringToPath(InInstance, "/interaction_profiles/microsoft/hand_interaction", &OutPath) == XR_SUCCESS;
+	OutKeyPrefixes.Add("OpenXRMsftHandInteraction");
+	OutHasHaptics.Add(false);
+	XrPath OutPath;
+	XrResult Result = xrStringToPath(InInstance, "/interaction_profiles/microsoft/hand_interaction", &OutPath);
+	OutPaths.Add(OutPath);
+	return Result == XR_SUCCESS;
 }
 
 #undef LOCTEXT_NAMESPACE
diff --git a/Engine/Plugins/Runtime/OpenXRMsftHandInteraction/Source/OpenXRMsftHandInteraction/Public/OpenXRMsftHandInteraction.h b/Engine/Plugins/Runtime/OpenXRMsftHandInteraction/Source/OpenXRMsftHandInteraction/Public/OpenXRMsftHandInteraction.h
--- a/Engine/Plugins/Runtime/OpenXRMsftHandInteraction/Source/OpenXRMsftHandInteraction/Public/OpenXRMsftHandInteraction.h
+++ b/Engine/Plugins/Runtime/OpenXRMsftHandInteraction/Source/OpenXRMsftHandInteraction/Public/OpenXRMsftHandInteraction.h
@@ -34,5 +34,5 @@ public:
 		return FString(TEXT("OpenXRMsftHandInteraction"));
 	}
 	virtual bool GetRequiredExtensions(TArray<const ANSICHAR*>& OutExtensions) override;
-	virtual bool GetInteractionProfile(XrInstance InInstance, FString& OutKeyPrefix, XrPath& OutPath, bool& OutHasHaptics) override;
+	virtual bool GetInteractionProfiles(XrInstance InInstance, TArray<FString>& OutKeyPrefixes, TArray<XrPath>& OutPaths, TArray<bool>& OutHasHaptics) override;
 };
