ensure() instead of check() if OpenXRPath cannot be converted to a string

diff --git a/Engine/Plugins/Runtime/OpenXR/Source/OpenXRHMD/Private/OpenXRCore.cpp b/Engine/Plugins/Runtime/OpenXR/Source/OpenXRHMD/Private/OpenXRCore.cpp
--- a/Engine/Plugins/Runtime/OpenXR/Source/OpenXRHMD/Private/OpenXRCore.cpp
+++ b/Engine/Plugins/Runtime/OpenXR/Source/OpenXRHMD/Private/OpenXRCore.cpp
@@ -41,7 +41,7 @@ FString FOpenXRPath::ToString() const
 	uint32 PathCount = 0;
 	char PathChars[XR_MAX_PATH_LENGTH];
 	XrResult Result = xrPathToString(Instance, Path, XR_MAX_PATH_LENGTH, &PathCount, PathChars);
-	check(XR_SUCCEEDED(Result));
+	ensure(XR_SUCCEEDED(Result));
 	return XR_SUCCEEDED(Result) ? FString::ConstructFromPtrSize(PathChars, PathCount - 1) : FString();
 }
 
