PR #12229: Add OpenXR extension plugin function for overriding input paths

diff --git a/Engine/Plugins/Runtime/OpenXR/Source/OpenXRHMD/Public/IOpenXRExtensionPlugin.h b/Engine/Plugins/Runtime/OpenXR/Source/OpenXRHMD/Public/IOpenXRExtensionPlugin.h
--- a/Engine/Plugins/Runtime/OpenXR/Source/OpenXRHMD/Public/IOpenXRExtensionPlugin.h
+++ b/Engine/Plugins/Runtime/OpenXR/Source/OpenXRHMD/Public/IOpenXRExtensionPlugin.h
@@ -112,6 +112,13 @@ public:
 	virtual void OnDestroySession() = 0;
 };
 
+struct FInputKeyOpenXRProperties
+{
+	FString InputKey;
+	FString InteractionProfile;
+	FString OpenXRPath;
+};
+
 class IOpenXRExtensionPlugin : public IModularFeature
 {
 public:
@@ -233,6 +240,16 @@ public:
 		return false;
 	}
 
+	/**
+	 * Set the output parameter to explicitly define an interaction profile and path for the given key.
+	 * The same key can contain multiple entries if the key is relevant to multiple interaction profiles.
+	 * If false is returned the overrides will be ignored.
+	 */
+	virtual bool GetInputKeyOverrides(TArray<FInputKeyOpenXRProperties>& OutOverrides)
+	{
+		return false;
+	}
+
 	/**
 	 * Set the output parameters to provide a path to an asset in the plugin content folder that visualizes
 	 * the controller in the hand represented by the user path.
diff --git a/Engine/Plugins/Runtime/OpenXR/Source/OpenXRInput/Private/OpenXRInput.cpp b/Engine/Plugins/Runtime/OpenXR/Source/OpenXRInput/Private/OpenXRInput.cpp
--- a/Engine/Plugins/Runtime/OpenXR/Source/OpenXRInput/Private/OpenXRInput.cpp
+++ b/Engine/Plugins/Runtime/OpenXR/Source/OpenXRInput/Private/OpenXRInput.cpp
@@ -390,6 +390,19 @@ bool FOpenXRInputPlugin::FOpenXRInput::BuildActions(XrSession Session)
 	Profiles.Add("OculusTouch", FInteractionProfile(FOpenXRPath("/interaction_profiles/oculus/touch_controller"), true));
 	Profiles.Add("ValveIndex", FInteractionProfile(FOpenXRPath("/interaction_profiles/valve/index_controller"), true));
 
+	// Query extension plugins for input key overrides
+	for (IOpenXRExtensionPlugin* Plugin : OpenXRHMD->GetExtensionPlugins())
+	{
+		TArray<FInputKeyOpenXRProperties> PluginInputOverrides;
+		if (Plugin->GetInputKeyOverrides(PluginInputOverrides))
+		{
+			for (const FInputKeyOpenXRProperties& PluginInputOverride : PluginInputOverrides)
+			{
+				InputsKeysToPropertiesMap.FindOrAdd(PluginInputOverride.InputKey).Add(PluginInputOverride);
+			}
+		}
+	}
+
 	// Query extension plugins for interaction profiles
 	{
 		TArray<FString> KeyPrefixes;
@@ -728,6 +741,22 @@ int32 FOpenXRInputPlugin::FOpenXRInput::SuggestBindings(TMap<FString, FInteracti
 
 bool FOpenXRInputPlugin::FOpenXRInput::SuggestBindingForKey(TMap<FString, FInteractionProfile>& Profiles, FOpenXRAction& Action, const FKey& InFKey, const TArray<UInputModifier*>& Modifiers, const TArray<UInputTrigger*>& Triggers)
 {
+	// Use profiles and path from overrides, if an entry exists for the key
+	FString InputKey = InFKey.ToString();
+	if (InputsKeysToPropertiesMap.Contains(InputKey))
+	{
+		for (const FInputKeyOpenXRProperties& InputProperties : InputsKeysToPropertiesMap[InputKey])
+		{
+			FInteractionProfile* Profile = Profiles.Find(InputProperties.InteractionProfile);
+			if (Profile == nullptr)
+			{
+				continue;
+			}
+			Profile->Bindings.Add(XrActionSuggestedBinding{ Action.Handle, FOpenXRPath(InputProperties.OpenXRPath) });
+		}
+		return true;
+	}
+
 	// Key names that are parseable into an OpenXR path have exactly 4 tokens
 	TArray<FString> Tokens;
 	if (InFKey.ToString().ParseIntoArray(Tokens, TEXT("_")) != EKeys::NUM_XR_KEY_TOKENS)
diff --git a/Engine/Plugins/Runtime/OpenXR/Source/OpenXRInput/Private/OpenXRInput.h b/Engine/Plugins/Runtime/OpenXR/Source/OpenXRInput/Private/OpenXRInput.h
--- a/Engine/Plugins/Runtime/OpenXR/Source/OpenXRInput/Private/OpenXRInput.h
+++ b/Engine/Plugins/Runtime/OpenXR/Source/OpenXRInput/Private/OpenXRInput.h
@@ -156,6 +156,7 @@ public:
 		TArray<FOpenXRAction> LegacyActions, EnhancedActions;
 		TMap<EControllerHand, FOpenXRController> Controllers;
 		TMap<FName, EControllerHand> MotionSourceToControllerHandMap;
+		TMap<FString, TArray<FInputKeyOpenXRProperties>> InputsKeysToPropertiesMap;
 
 		// Holds Enhanced Input Mapping Contexts pulled from the Enhanced Input Settings and mapped to their priority
 		// These will be converted into OpenXR action sets
