Fix OpenXR foveation crash on map change.

diff --git a/Engine/Plugins/Runtime/OpenXR/Source/OpenXRHMD/Private/OpenXRHMD.cpp b/Engine/Plugins/Runtime/OpenXR/Source/OpenXRHMD/Private/OpenXRHMD.cpp
--- a/Engine/Plugins/Runtime/OpenXR/Source/OpenXRHMD/Private/OpenXRHMD.cpp
+++ b/Engine/Plugins/Runtime/OpenXR/Source/OpenXRHMD/Private/OpenXRHMD.cpp
@@ -1689,6 +1689,11 @@ FOpenXRHMD::FOpenXRHMD(const FAutoRegister& AutoRegister, XrInstance InInstance,
 
 FOpenXRHMD::~FOpenXRHMD()
 {
+	if (bRuntimeFoveationSupported)
+	{
+		GVRSImageManager.UnregisterExternalImageGenerator(FBFoveationImageGenerator.Get());
+		FBFoveationImageGenerator.Reset();
+	}
 	DestroySession();
 }
 
@@ -2548,22 +2553,13 @@ void FOpenXRHMD::OnBeginPlay(FWorldContext& InWorldContext)
 
 	const UOpenXRHMDSettings* Settings = GetDefault<UOpenXRHMDSettings>();
 	bRuntimeFoveationSupported = bFoveationExtensionSupported && (Settings != nullptr ? Settings->bIsFBFoveationEnabled : false);
-	if (bRuntimeFoveationSupported)
+	if (bRuntimeFoveationSupported && !FBFoveationImageGenerator.IsValid())
 	{
 		FBFoveationImageGenerator = MakeUnique<FFBFoveationImageGenerator>(bRuntimeFoveationSupported, Instance, this, bIsMobileMultiViewEnabled);
 		GVRSImageManager.RegisterExternalImageGenerator(FBFoveationImageGenerator.Get());
 	}
 }
 
-void FOpenXRHMD::OnEndPlay(FWorldContext& InWorldContext)
-{
-	if (bRuntimeFoveationSupported)
-	{
-		GVRSImageManager.UnregisterExternalImageGenerator(FBFoveationImageGenerator.Get());
-		FBFoveationImageGenerator.Reset();
-	}
-}
-
 IStereoRenderTargetManager* FOpenXRHMD::GetRenderTargetManager()
 {
 	return this;
diff --git a/Engine/Plugins/Runtime/OpenXR/Source/OpenXRHMD/Private/OpenXRHMD.h b/Engine/Plugins/Runtime/OpenXR/Source/OpenXRHMD/Private/OpenXRHMD.h
--- a/Engine/Plugins/Runtime/OpenXR/Source/OpenXRHMD/Private/OpenXRHMD.h
+++ b/Engine/Plugins/Runtime/OpenXR/Source/OpenXRHMD/Private/OpenXRHMD.h
@@ -432,7 +432,6 @@ public:
 
 	/** IXRTrackingSystem */
 	virtual void OnBeginPlay(FWorldContext& InWorldContext) override;
-	virtual void OnEndPlay(FWorldContext& InWorldContext) override;
 
 	/** IStereoLayers */
 	virtual bool ShouldCopyDebugLayersToSpectatorScreen() const override { return true; }
