VisionOS passthrough edge artifact
-The order in which we would acquire the new depth texture and start using it for rendering was subject to a race condition between tasks queued for the RHI and Render threads.  This race was resulting in rendering to the previous depth buffer on perhaps 1/50 frames on VisionOS.  VisionOS does not have an RHI thread, but the two separate task queues appeared to be processed in an inconsistent order (perhaps based on when they were enqueued from the game and render threads. It is unknown if this problem actually happens in the case where there is a render thread.  The depth buffer is not always used, and off by one errors with it are probably generally more subtle.  They became obvious on VisionOS because any 32x32 pixel block where the depth is always 0 is not displayed at all leading to these sawtooth edges on mixed reality objects rendering against the passthrough background.
-Here I move the acquire up to the same place where we acquire the color texture on the Game Thread.

diff --git a/Engine/Platforms/VisionOS/Config/VisionOSEngine.ini b/Engine/Platforms/VisionOS/Config/VisionOSEngine.ini
new file mode 100644
--- /dev/null
+++ b/Engine/Platforms/VisionOS/Config/VisionOSEngine.ini
@@ -0,0 +1,16 @@
+; NOTE: most values should be added to BaseVisionOSEngine.ini unless there is a very specific reason for it
+; items added here will take priority over the project's DefaultEngine.ini. Use the project's VisionOSEngine.ini to override values here.
+
+; These settings should be considered 'recommended platform defaults'
+
+[/Script/Engine.RendererSettings]
+; These settings are always required for VisionOS, so we override any other project settings.
+vr.InstancedStereo=False
+vr.MobileMultiView=False
+xr.OpenXRAcquireMode=1
+
+// These are required for .mixed mode, aka AR
+;r.Mobile.PropagateAlpha=1
+;r.PostProcessing.PropagateAlpha=1
+;r.AlphaInvertPass=1
+
diff --git a/Engine/Plugins/Runtime/OpenXR/Source/OpenXRHMD/Private/OpenXRHMD.cpp b/Engine/Plugins/Runtime/OpenXR/Source/OpenXRHMD/Private/OpenXRHMD.cpp
--- a/Engine/Plugins/Runtime/OpenXR/Source/OpenXRHMD/Private/OpenXRHMD.cpp
+++ b/Engine/Plugins/Runtime/OpenXR/Source/OpenXRHMD/Private/OpenXRHMD.cpp
@@ -2583,6 +2583,24 @@ int32 FOpenXRHMD::AcquireColorTexture()
 	return 0;
 }
 
+int32 FOpenXRHMD::AcquireDepthTexture()
+{
+	check(IsInGameThread());
+	if (Session)
+	{
+		const FXRSwapChainPtr& DepthSwapchain = PipelinedLayerStateRendering.DepthSwapchain;
+		if (DepthSwapchain)
+		{
+			if (bIsAcquireOnAnyThreadSupported)
+			{
+				DepthSwapchain->IncrementSwapChainIndex_RHIThread();
+			}
+			return DepthSwapchain->GetSwapChainIndex_RHIThread();
+		}
+	}
+	return 0;
+}
+
 bool FOpenXRHMD::AllocateRenderTargetTextures(uint32 SizeX, uint32 SizeY, uint8 Format, uint32 NumMips, ETextureCreateFlags Flags, ETextureCreateFlags TargetableTextureFlags, TArray<FTextureRHIRef>& OutTargetableTextures, TArray<FTextureRHIRef>& OutShaderResourceTextures, uint32 NumSamples)
 {
 	check(IsInRenderingThread());
@@ -2762,11 +2780,6 @@ bool FOpenXRHMD::AllocateDepthTexture(uint32 Index, uint32 SizeX, uint32 SizeY,
 
 	const ETextureCreateFlags UnifiedCreateFlags = Flags | TargetableTextureFlags;
 	ensure(EnumHasAllFlags(UnifiedCreateFlags, TexCreate_DepthStencilTargetable)); // We can't use the depth swapchain w/o this flag
-	if (bIsAcquireOnAnyThreadSupported)
-	{
-		DepthSwapchain->IncrementSwapChainIndex_RHIThread();
-	}
-
 	const FRHITexture* const DepthSwapchainTexture = DepthSwapchain->GetTexture2DArray() ? DepthSwapchain->GetTexture2DArray() : DepthSwapchain->GetTexture2D();
 	const FRHITextureDesc& DepthSwapchainDesc = DepthSwapchainTexture->GetDesc();
 
diff --git a/Engine/Plugins/Runtime/OpenXR/Source/OpenXRHMD/Private/OpenXRHMD.h b/Engine/Plugins/Runtime/OpenXR/Source/OpenXRHMD/Private/OpenXRHMD.h
--- a/Engine/Plugins/Runtime/OpenXR/Source/OpenXRHMD/Private/OpenXRHMD.h
+++ b/Engine/Plugins/Runtime/OpenXR/Source/OpenXRHMD/Private/OpenXRHMD.h
@@ -423,6 +423,7 @@ public:
 	virtual void CalculateRenderTargetSize(const FViewport& Viewport, uint32& InOutSizeX, uint32& InOutSizeY) override;
 	virtual bool AllocateRenderTargetTextures(uint32 SizeX, uint32 SizeY, uint8 Format, uint32 NumLayers, ETextureCreateFlags Flags, ETextureCreateFlags TargetableTextureFlags, TArray<FTextureRHIRef>& OutTargetableTextures, TArray<FTextureRHIRef>& OutShaderResourceTextures, uint32 NumSamples = 1) override;
 	virtual int32 AcquireColorTexture() override final;
+	virtual int32 AcquireDepthTexture() override final;
 	virtual bool AllocateDepthTexture(uint32 Index, uint32 SizeX, uint32 SizeY, uint8 Format, uint32 NumMips, ETextureCreateFlags InTexFlags, ETextureCreateFlags TargetableTextureFlags, FTextureRHIRef& OutTargetableTexture, FTextureRHIRef& OutShaderResourceTexture, uint32 NumSamples = 1) override final;
 	virtual bool ReconfigureForShaderPlatform(EShaderPlatform NewShaderPlatform) override;
 	virtual EPixelFormat GetActualColorSwapchainFormat() const override { return static_cast<EPixelFormat>(LastActualColorSwapchainFormat); }
@@ -587,4 +588,4 @@ private:
 	bool					bxrGetSystemPropertiesSuccessful;
 };
 
-ENUM_CLASS_FLAGS(FOpenXRHMD::EOpenXRLayerStateFlags);
\ No newline at end of file
+ENUM_CLASS_FLAGS(FOpenXRHMD::EOpenXRLayerStateFlags);
diff --git a/Engine/Source/Runtime/Engine/Private/Slate/SceneViewport.cpp b/Engine/Source/Runtime/Engine/Private/Slate/SceneViewport.cpp
--- a/Engine/Source/Runtime/Engine/Private/Slate/SceneViewport.cpp
+++ b/Engine/Source/Runtime/Engine/Private/Slate/SceneViewport.cpp
@@ -1826,6 +1826,8 @@ void FSceneViewport::EnqueueBeginRenderFrame(const bool bShouldPresent)
 			// We need to acquire a buffered texture from either the new RT or the existing one
 			int32 TextureIndex = StereoRenderTargetManager->AcquireColorTexture();
 			CurrentBufferedTargetIndex = TextureIndex < 0 ? CurrentBufferedTargetIndex : TextureIndex;
+			
+			StereoRenderTargetManager->AcquireDepthTexture();
 		}
 	}
 
diff --git a/Engine/Source/Runtime/Engine/Public/StereoRenderTargetManager.h b/Engine/Source/Runtime/Engine/Public/StereoRenderTargetManager.h
--- a/Engine/Source/Runtime/Engine/Public/StereoRenderTargetManager.h
+++ b/Engine/Source/Runtime/Engine/Public/StereoRenderTargetManager.h
@@ -95,6 +95,13 @@ public:
 	 * @return				the index of the texture in the array returned by AllocateRenderTargetTexture.
 	 */
 	virtual int32 AcquireColorTexture() { return -1; }
+	
+	/**
+	 * Acquires the next available depth texture.
+	 * 
+	 * @return				the index of the texture in the array returned by AllocateRenderTargetTexture.
+	 */
+	virtual int32 AcquireDepthTexture() { return -1; }
 
 	/**
 	 * Allocates a depth texture.
@@ -132,4 +139,4 @@ public:
 	virtual bool ReconfigureForShaderPlatform(EShaderPlatform NewShaderPlatform) { return false; };
 
 	static EPixelFormat GetStereoLayerPixelFormat() { return PF_B8G8R8A8; }
-};
\ No newline at end of file
+};
diff --git a/Templates/TP_VirtualRealityBP/Config/VisionOS/VisionOSEngine.ini b/Templates/TP_VirtualRealityBP/Config/VisionOS/VisionOSEngine.ini
--- a/Templates/TP_VirtualRealityBP/Config/VisionOS/VisionOSEngine.ini
+++ b/Templates/TP_VirtualRealityBP/Config/VisionOS/VisionOSEngine.ini
@@ -3,8 +3,10 @@
 // These settings are currently required for all visionOS projects.
 vr.InstancedStereo=False
 vr.MobileMultiView=False
+xr.OpenXRAcquireMode=1
 
-//Enable these settings for .mixed immesion mode.
+//Enable these settings for .mixed immesion mode, aka passthrough, AR.
+//You will also need to edit UESwift.swift
 //r.Mobile.PropagateAlpha=1
 //r.PostProcessing.PropagateAlpha=1
 //r.AlphaInvertPass=1
\ No newline at end of file
