Add OpenXR support for XR_KHR_composition_layer_equirect2

diff --git a/Engine/Plugins/Runtime/OpenXR/Source/OpenXRHMD/Private/OpenXRHMD.cpp b/Engine/Plugins/Runtime/OpenXR/Source/OpenXRHMD/Private/OpenXRHMD.cpp
--- a/Engine/Plugins/Runtime/OpenXR/Source/OpenXRHMD/Private/OpenXRHMD.cpp
+++ b/Engine/Plugins/Runtime/OpenXR/Source/OpenXRHMD/Private/OpenXRHMD.cpp
@@ -2973,7 +2973,9 @@ void FOpenXRHMD::SetupFrameLayers_RenderThread(FRHICommandListImmediate& RHICmdL
 		XrSpace Space = Layer.Desc.PositionType == ELayerType::FaceLocked ?
 			DeviceSpaces[HMDDeviceId].Space : PipelinedFrameStateRendering.TrackingSpace->Handle;
 		
-		TArray<FXrCompositionLayerUnion> Headers = Layer.CreateOpenXRLayer(InvTrackingToWorld, WorldToMeters, Space);
+		EOpenXRLayerCreationFlags LayerCreationFlags = EOpenXRLayerCreationFlags::None;
+		LayerCreationFlags |= IsExtensionEnabled(XR_KHR_COMPOSITION_LAYER_EQUIRECT2_EXTENSION_NAME) ? EOpenXRLayerCreationFlags::EquirectLayer2Supported : EOpenXRLayerCreationFlags::None;
+		TArray<FXrCompositionLayerUnion> Headers = Layer.CreateOpenXRLayer(InvTrackingToWorld, WorldToMeters, Space, LayerCreationFlags);
 		PipelinedLayerStateRendering.NativeOverlays.Append(Headers);
 		UpdateLayerSwapchainTexture(Layer, RHICmdList);
 	}
@@ -2986,10 +2988,10 @@ void FOpenXRHMD::SetupFrameLayers_RenderThread(FRHICommandListImmediate& RHICmdL
 
 void FOpenXRHMD::ConfigureLayerSwapchain(FOpenXRLayer& Layer, TArray<FOpenXRLayer>& BackupLayers)
 {
+	const bool bIsEquirectLayerSupported = IsExtensionEnabled(XR_KHR_COMPOSITION_LAYER_EQUIRECT_EXTENSION_NAME) || IsExtensionEnabled(XR_KHR_COMPOSITION_LAYER_EQUIRECT2_EXTENSION_NAME);
 	// OpenXR currently supports only Quad layers unless the cylinder and equirect extensions are enabled.
 	if ((Layer.Desc.HasShape<FCylinderLayer>() && IsExtensionEnabled(XR_KHR_COMPOSITION_LAYER_CYLINDER_EXTENSION_NAME)) ||
-		(Layer.Desc.HasShape<FEquirectLayer>() && IsExtensionEnabled(XR_KHR_COMPOSITION_LAYER_EQUIRECT_EXTENSION_NAME)) ||
-		Layer.Desc.HasShape<FQuadLayer>())
+		(Layer.Desc.HasShape<FEquirectLayer>() && bIsEquirectLayerSupported) || Layer.Desc.HasShape<FQuadLayer>())
 	{
 		if (Layer.Desc.IsVisible())
 		{
diff --git a/Engine/Plugins/Runtime/OpenXR/Source/OpenXRHMD/Private/OpenXRHMDModule.cpp b/Engine/Plugins/Runtime/OpenXR/Source/OpenXRHMD/Private/OpenXRHMDModule.cpp
--- a/Engine/Plugins/Runtime/OpenXR/Source/OpenXRHMD/Private/OpenXRHMDModule.cpp
+++ b/Engine/Plugins/Runtime/OpenXR/Source/OpenXRHMD/Private/OpenXRHMDModule.cpp
@@ -561,6 +561,7 @@ bool FOpenXRHMDModule::GetOptionalExtensions(TArray<const ANSICHAR*>& OutExtensi
 	OutExtensions.Add(XR_KHR_COMPOSITION_LAYER_DEPTH_EXTENSION_NAME);
 	OutExtensions.Add(XR_KHR_COMPOSITION_LAYER_CYLINDER_EXTENSION_NAME);
 	OutExtensions.Add(XR_KHR_COMPOSITION_LAYER_EQUIRECT_EXTENSION_NAME);
+	OutExtensions.Add(XR_KHR_COMPOSITION_LAYER_EQUIRECT2_EXTENSION_NAME);
 	OutExtensions.Add(XR_KHR_COMPOSITION_LAYER_COLOR_SCALE_BIAS_EXTENSION_NAME);
 	OutExtensions.Add(XR_VARJO_QUAD_VIEWS_EXTENSION_NAME);
 	OutExtensions.Add(XR_KHR_VISIBILITY_MASK_EXTENSION_NAME);
diff --git a/Engine/Plugins/Runtime/OpenXR/Source/OpenXRHMD/Private/OpenXRHMD_Layer.cpp b/Engine/Plugins/Runtime/OpenXR/Source/OpenXRHMD/Private/OpenXRHMD_Layer.cpp
--- a/Engine/Plugins/Runtime/OpenXR/Source/OpenXRHMD/Private/OpenXRHMD_Layer.cpp
+++ b/Engine/Plugins/Runtime/OpenXR/Source/OpenXRHMD/Private/OpenXRHMD_Layer.cpp
@@ -79,7 +79,7 @@ FVector2D FOpenXRLayer::GetLeftQuadSize() const
 	}
 	return Desc.QuadSize;
 }
-TArray<FXrCompositionLayerUnion> FOpenXRLayer::CreateOpenXRLayer(FTransform InvTrackingToWorld, float WorldToMeters, XrSpace Space) const
+TArray<FXrCompositionLayerUnion> FOpenXRLayer::CreateOpenXRLayer(FTransform InvTrackingToWorld, float WorldToMeters, XrSpace Space, EOpenXRLayerCreationFlags CreationFlags) const
 {
 	TArray<FXrCompositionLayerUnion> Headers;
 
@@ -98,7 +98,14 @@ TArray<FXrCompositionLayerUnion> FOpenXRLayer::CreateOpenXRLayer(FTransform InvT
 	}
 	else if (Desc.HasShape<FEquirectLayer>())
 	{
-		CreateOpenXREquirectLayer(bIsStereo, bNoAlpha, PositionTransform, WorldToMeters, Space, Headers);
+		if(EnumHasAllFlags(CreationFlags, EOpenXRLayerCreationFlags::EquirectLayer2Supported))
+		{
+			CreateOpenXREquirect2Layer(bIsStereo, bNoAlpha, PositionTransform, WorldToMeters, Space, Headers);
+		}
+		else
+		{
+			CreateOpenXREquirectLayer(bIsStereo, bNoAlpha, PositionTransform, WorldToMeters, Space, Headers);
+		}
 	}
 
 	return Headers;
@@ -171,6 +178,71 @@ void FOpenXRLayer::CreateOpenXRQuadLayer(bool bIsStereo, bool bNoAlpha, FTransfo
 	}
 }
 
+void FOpenXRLayer::SetupEquirect2(FVector2D UVScale, FVector2D UVBias, FVector2D UVPosition, FVector2D UVSize, FTransform PositionTransform, float WorldToMeters, XrCompositionLayerEquirect2KHR& Equirect2) const
+{
+	const FVector2D AdjustedPosition = FVector2D(
+		(-UVBias.X + UVPosition.X) / UVScale.X, 
+		(-UVBias.Y + UVPosition.Y) / UVScale.Y);
+	const FVector2D AdjustedSize = FVector2D(UVSize.X / UVScale.X, UVSize.Y / UVScale.Y);
+
+	const float CentralHorizontalAngle = (UE_PI * 2.0f) * AdjustedSize.X;
+	const float UpperVerticalAngle = (UE_PI / 2.0f) - ((1.0f - AdjustedPosition.Y - AdjustedSize.Y) * UE_PI);
+	const float LowerVerticalAngle = (-UE_PI / 2.0f) + (AdjustedPosition.Y * UE_PI);
+
+	Equirect2.centralHorizontalAngle = CentralHorizontalAngle;
+	Equirect2.upperVerticalAngle = UpperVerticalAngle; 
+	Equirect2.lowerVerticalAngle = LowerVerticalAngle;
+	Equirect2.pose = ToXrPose(Desc.Transform * PositionTransform, WorldToMeters);
+}
+
+void FOpenXRLayer::CreateOpenXREquirect2Layer(bool bIsStereo, bool bNoAlpha, FTransform PositionTransform, float WorldToMeters, XrSpace Space, TArray<FXrCompositionLayerUnion>& Headers) const
+{
+	const FEquirectLayer& EquirectProps = Desc.GetShape<FEquirectLayer>();
+
+	XrCompositionLayerEquirect2KHR Equirect2 = { XR_TYPE_COMPOSITION_LAYER_EQUIRECT2_KHR, nullptr };
+	Equirect2.layerFlags = bNoAlpha ? 0 : XR_COMPOSITION_LAYER_UNPREMULTIPLIED_ALPHA_BIT |
+		XR_COMPOSITION_LAYER_BLEND_TEXTURE_SOURCE_ALPHA_BIT;
+	Equirect2.space = Space;
+	Equirect2.subImage.imageArrayIndex = 0;
+	FXrCompositionLayerUnion LayerUnion;
+
+	if (RightEye.Swapchain.IsValid())
+	{
+		const FVector2D Size = EquirectProps.RightUVRect.Max - EquirectProps.RightUVRect.Min;
+		SetupEquirect2(EquirectProps.RightScale,
+					   EquirectProps.RightBias,
+					   EquirectProps.RightUVRect.Min,
+					   Size,
+					   PositionTransform,
+					   WorldToMeters,
+					   Equirect2);
+
+		LayerUnion.Equirect2 = Equirect2;
+		LayerUnion.Equirect2.eyeVisibility = bIsStereo ? XR_EYE_VISIBILITY_RIGHT : XR_EYE_VISIBILITY_BOTH;
+		LayerUnion.Equirect2.subImage.imageRect = ToXrRect(GetRightViewportSize());
+		LayerUnion.Equirect2.subImage.swapchain = static_cast<FOpenXRSwapchain*>(RightEye.Swapchain.Get())->GetHandle();
+		Headers.Add(LayerUnion);
+	}
+
+	if (LeftEye.Swapchain.IsValid())
+	{
+		const FVector2D Size = EquirectProps.LeftUVRect.Max - EquirectProps.LeftUVRect.Min;
+		SetupEquirect2(EquirectProps.LeftScale,
+					   EquirectProps.LeftBias,
+					   EquirectProps.LeftUVRect.Min,
+					   Size,
+					   PositionTransform,
+					   WorldToMeters,
+					   Equirect2);
+
+		LayerUnion.Equirect2 = Equirect2;
+		LayerUnion.Equirect2.eyeVisibility = XR_EYE_VISIBILITY_LEFT;
+		LayerUnion.Equirect2.subImage.imageRect = ToXrRect(GetLeftViewportSize());
+		LayerUnion.Equirect2.subImage.swapchain = static_cast<FOpenXRSwapchain*>(LeftEye.Swapchain.Get())->GetHandle();
+		Headers.Add(LayerUnion);
+	}
+}
+
 void FOpenXRLayer::CreateOpenXREquirectLayer(bool bIsStereo, bool bNoAlpha, FTransform PositionTransform, float WorldToMeters, XrSpace Space, TArray<FXrCompositionLayerUnion>& Headers) const
 {
 	XrCompositionLayerEquirectKHR Equirect = { XR_TYPE_COMPOSITION_LAYER_EQUIRECT_KHR, nullptr };
diff --git a/Engine/Plugins/Runtime/OpenXR/Source/OpenXRHMD/Private/OpenXRHMD_Layer.h b/Engine/Plugins/Runtime/OpenXR/Source/OpenXRHMD/Private/OpenXRHMD_Layer.h
--- a/Engine/Plugins/Runtime/OpenXR/Source/OpenXRHMD/Private/OpenXRHMD_Layer.h
+++ b/Engine/Plugins/Runtime/OpenXR/Source/OpenXRHMD/Private/OpenXRHMD_Layer.h
@@ -49,12 +49,14 @@ struct FOpenXRLayer
 	FIntRect GetLeftViewportSize() const;
 	FVector2D GetLeftQuadSize() const;
 
-	TArray<FXrCompositionLayerUnion> CreateOpenXRLayer(FTransform InvTrackingToWorld, float WorldToMeters, XrSpace Space) const;
+	TArray<FXrCompositionLayerUnion> CreateOpenXRLayer(FTransform InvTrackingToWorld, float WorldToMeters, XrSpace Space, EOpenXRLayerCreationFlags CreationFlags) const;
 	
 private:
 	void CreateOpenXRQuadLayer(bool bIsStereo, bool bNoAlpha, FTransform PositionTransform, float WorldToMeters, XrSpace Space, TArray<FXrCompositionLayerUnion>& Headers) const;
 	void CreateOpenXRCylinderLayer(bool bIsStereo, bool bNoAlpha, FTransform PositionTransform, float WorldToMeters, XrSpace Space, TArray<FXrCompositionLayerUnion>& Headers) const;
 	void CreateOpenXREquirectLayer(bool bIsStereo, bool bNoAlpha, FTransform PositionTransform, float WorldToMeters, XrSpace Space, TArray<FXrCompositionLayerUnion>& Headers) const;
+	void CreateOpenXREquirect2Layer(bool bIsStereo, bool bNoAlpha, FTransform PositionTransform, float WorldToMeters, XrSpace Space, TArray<FXrCompositionLayerUnion>& Headers) const;
+	void SetupEquirect2(FVector2D Scale, FVector2D Bias, FVector2D Position, FVector2D Size, FTransform PositionTransform, float WorldToMeters, XrCompositionLayerEquirect2KHR& Equirect2) const;
 
 };
 
diff --git a/Engine/Plugins/Runtime/OpenXR/Source/OpenXRHMD/Public/OpenXRCore.h b/Engine/Plugins/Runtime/OpenXR/Source/OpenXRHMD/Public/OpenXRCore.h
--- a/Engine/Plugins/Runtime/OpenXR/Source/OpenXRHMD/Public/OpenXRCore.h
+++ b/Engine/Plugins/Runtime/OpenXR/Source/OpenXRHMD/Public/OpenXRCore.h
@@ -142,8 +142,17 @@ union FXrCompositionLayerUnion
 	XrCompositionLayerQuad Quad;
 	XrCompositionLayerCylinderKHR Cylinder;
 	XrCompositionLayerEquirectKHR Equirect;
+	XrCompositionLayerEquirect2KHR Equirect2;
 };
 
+enum class EOpenXRLayerCreationFlags : uint32
+{
+	None = 0u,
+	EquirectLayer2Supported = (1u << 0),
+};
+
+ENUM_CLASS_FLAGS(EOpenXRLayerCreationFlags);
+
 /**
  * XrPath wrapper with convenience functions
  */
