PR #12399: Reset OpenXRHMD SwapChains declared as static on update
Ensure that static swapchains are marked reset on both OpenXR native stereo layers, and engine manager layers, when using OpenXR. Previously, a mismatch between the two could cause a crash when moving layers not declared as live/dynamic, since this requires a swapchain reset.

diff --git a/Engine/Plugins/Runtime/OpenXR/Source/OpenXRHMD/Private/OpenXRHMD.cpp b/Engine/Plugins/Runtime/OpenXR/Source/OpenXRHMD/Private/OpenXRHMD.cpp
--- a/Engine/Plugins/Runtime/OpenXR/Source/OpenXRHMD/Private/OpenXRHMD.cpp
+++ b/Engine/Plugins/Runtime/OpenXR/Source/OpenXRHMD/Private/OpenXRHMD.cpp
@@ -4414,20 +4414,26 @@ void FOpenXRHMD::DrawVisibleAreaMesh(class FRHICommandList& RHICmdList, int32 Vi
 
 void FOpenXRHMD::UpdateLayer(FOpenXRLayer& ManagerLayer, uint32 LayerId, bool bIsValid)
 {
+	const bool bStaticSwapchain = !(ManagerLayer.Desc.Flags & IStereoLayers::LAYER_FLAG_TEX_CONTINUOUS_UPDATE);
+	if (bStaticSwapchain)
+	{
+		ManagerLayer.RightEye.Swapchain.Reset();
+		ManagerLayer.LeftEye.Swapchain.Reset();
+	}
+
 	ENQUEUE_RENDER_COMMAND(UpdateLayer)(
 		[
 			this,
-			Flags = ManagerLayer.Desc.Flags,
 			bUpdateRightEyeTexture = ManagerLayer.RightEye.bUpdateTexture,
 			bUpdateLeftEyeTexture = ManagerLayer.LeftEye.bUpdateTexture,
-			LayerId
+			LayerId,
+			bStaticSwapchain
 		](FRHICommandList&)
 	{
 	for (FOpenXRLayer& NativeLayer : NativeLayers)
 		{
 			if (NativeLayer.GetLayerId() == LayerId)
 			{
-				const bool bStaticSwapchain = !(Flags & IStereoLayers::LAYER_FLAG_TEX_CONTINUOUS_UPDATE);
 				NativeLayer.RightEye.bUpdateTexture = bUpdateRightEyeTexture;
 				NativeLayer.LeftEye.bUpdateTexture = bUpdateLeftEyeTexture;
 				if (bStaticSwapchain)
