Can't enable inverse alpha
-r.AlphaInvertPass will now enable the pass for all renders, as it did previously.
-OpenXR.AlphaInvertPass will do it specifically for the XR views.
-In general it is a good idea to add the pass only for renders where it is needed, because it does have a performance cost.  The mobile deferred case is further optimized because it is part of the upscale pass (and always on - never controlled by cvar).

diff --git a/Engine/Plugins/Runtime/OpenXR/Source/OpenXRHMD/Private/OpenXRHMD.cpp b/Engine/Plugins/Runtime/OpenXR/Source/OpenXRHMD/Private/OpenXRHMD.cpp
--- a/Engine/Plugins/Runtime/OpenXR/Source/OpenXRHMD/Private/OpenXRHMD.cpp
+++ b/Engine/Plugins/Runtime/OpenXR/Source/OpenXRHMD/Private/OpenXRHMD.cpp
@@ -148,6 +148,13 @@ static TAutoConsoleVariable<bool> CVarOpenXRLateUpdateDeviceLocationsAfterReflec
 	TEXT("If you aren't using planar reflections in your project, you can safely enable this to get late update poses as late as possible.\n"),
 	ECVF_Default);
 
+static TAutoConsoleVariable<bool> CVarOpenXRAlphaInvertPass(
+	TEXT("OpenXR.AlphaInvertPass"),
+	0,
+	TEXT("Whether to run a render pass to un-invert the alpha value from unreal standard to the much more common standard where alpha 0 is fully transparent and alpha 1 is fully opaque."
+	TEXT("This cvar specifically enables the pass for the main XR view.  There is a more general r.ALphaInvertPass which enables it for all renders.")),
+	ECVF_RenderThreadSafe);
+
 namespace {
 	static TSet<XrViewConfigurationType> SupportedViewConfigurations{ XR_VIEW_CONFIGURATION_TYPE_PRIMARY_STEREO, XR_VIEW_CONFIGURATION_TYPE_PRIMARY_QUAD_VARJO };
 
@@ -1436,12 +1443,12 @@ void FOpenXRHMD::SetupViewFamily(FSceneViewFamily& InViewFamily)
 	if (InViewFamily.Scene)
 	{
 		EShaderPlatform Platform = InViewFamily.Scene->GetShaderPlatform();
-		static FShaderPlatformCachedIniValue<bool> AlphaInvertPassIniValue(TEXT("r.AlphaInvertPass"));
+		static FShaderPlatformCachedIniValue<bool> AlphaInvertPassIniValue(TEXT("OpenXR.AlphaInvertPass"));
 		InViewFamily.EngineShowFlags.AlphaInvert = AlphaInvertPassIniValue.Get(Platform);
 	}
 	else
 	{
-		static auto CVarAlphaInvertPassLocal = IConsoleManager::Get().FindConsoleVariable(TEXT("r.AlphaInvertPass"));
+		static auto CVarAlphaInvertPassLocal = IConsoleManager::Get().FindConsoleVariable(TEXT("OpenXR.AlphaInvertPass"));
 		InViewFamily.EngineShowFlags.AlphaInvert = (CVarAlphaInvertPassLocal != nullptr) && CVarAlphaInvertPassLocal->GetBool();
 	}
 
diff --git a/Engine/Source/Runtime/Engine/Private/ShowFlags.cpp b/Engine/Source/Runtime/Engine/Private/ShowFlags.cpp
--- a/Engine/Source/Runtime/Engine/Private/ShowFlags.cpp
+++ b/Engine/Source/Runtime/Engine/Private/ShowFlags.cpp
@@ -480,6 +480,14 @@ void EngineShowFlagOverride(EShowFlagInitMode ShowFlagInitMode, EViewModeIndex V
 		}
 	}
 
+	{
+		static const auto ICVar = IConsoleManager::Get().FindTConsoleVariableDataBool(TEXT("r.AlphaInvertPass"));
+		if (ICVar->GetValueOnGameThread())
+		{
+			EngineShowFlags.SetAlphaInvert(true); 
+		}
+	}
+
 	if (!IsRayTracingEnabled())
 	{
 		EngineShowFlags.SetPathTracing(false);
diff --git a/Engine/Source/Runtime/Renderer/Private/PostProcess/AlphaInvert.cpp b/Engine/Source/Runtime/Renderer/Private/PostProcess/AlphaInvert.cpp
--- a/Engine/Source/Runtime/Renderer/Private/PostProcess/AlphaInvert.cpp
+++ b/Engine/Source/Runtime/Renderer/Private/PostProcess/AlphaInvert.cpp
@@ -11,10 +11,12 @@
 
 namespace AlphaInvert
 {
-	static TAutoConsoleVariable<int32> CVarAlphaInvertPass(
+	static TAutoConsoleVariable<bool> CVarAlphaInvertPass(
 		TEXT("r.AlphaInvertPass"),
 		0,
-		TEXT("Whether to run a render pass to un-invert the alpha value from unreal standard to the much more common standard where alpha 0 is fully transparent and alpha 1 is fully opaque."),
+		TEXT("Whether to run a render pass to un-invert the alpha value from unreal standard to the much more common standard where alpha 0 is fully transparent and alpha 1 is fully opaque.")
+		TEXT("This cvar attempts to affect all renders, not only the main view.")
+		TEXT("If your project does multiple renders which do not all need alpha inverted it would be more performant to find or implement a narrower version of it for your specific purpose (eg OpenXR.AlphaInvertPass)."),
 		ECVF_RenderThreadSafe);
 
 	class FAlphaInvertPS : public FScreenPS
