More fine-grained scopes in OpenXRHMD.

diff --git a/Engine/Plugins/Runtime/OpenXR/Source/OpenXRHMD/Private/OpenXRHMD.cpp b/Engine/Plugins/Runtime/OpenXR/Source/OpenXRHMD/Private/OpenXRHMD.cpp
--- a/Engine/Plugins/Runtime/OpenXR/Source/OpenXRHMD/Private/OpenXRHMD.cpp
+++ b/Engine/Plugins/Runtime/OpenXR/Source/OpenXRHMD/Private/OpenXRHMD.cpp
@@ -3418,21 +3418,27 @@ void FOpenXRHMD::OnBeginSimulation_GameThread()
 
 	ensure(IsInGameThread());
 
-	SCOPED_NAMED_EVENT(WaitFrame, FColor::Red);
+	SCOPED_NAMED_EVENT(OpenXrWaitFrame, FColor::Red);
 
 	XrFrameWaitInfo WaitInfo;
 	WaitInfo.type = XR_TYPE_FRAME_WAIT_INFO;
 	WaitInfo.next = nullptr;
 
 	XrFrameState FrameState{XR_TYPE_FRAME_STATE};
-	for (IOpenXRExtensionPlugin* Module : ExtensionPlugins)
 	{
-		FrameState.next = Module->OnWaitFrame(Session, FrameState.next);
+		SCOPED_NAMED_EVENT(PluginsOnWaitFrame, FColor::Red);
+		for (IOpenXRExtensionPlugin* Module : ExtensionPlugins)
+		{
+			FrameState.next = Module->OnWaitFrame(Session, FrameState.next);
+		}
 	}
 	static int WaitCount = 0;
 	++WaitCount;
 	UE_LOG(LogHMD, VeryVerbose, TEXT("%s WF_%i xrWaitFrame Calling..."), HMDThreadString(), WaitCount);
-	XR_ENSURE(xrWaitFrame(Session, &WaitInfo, &FrameState));
+	{
+		SCOPED_NAMED_EVENT(xrWaitFrame, FColor::Red);
+		XR_ENSURE(xrWaitFrame(Session, &WaitInfo, &FrameState));
+	}
 	UE_LOG(LogHMD, VeryVerbose, TEXT("%s WF_%i xrWaitFrame Complete"), HMDThreadString(), WaitCount);
 
 	// The pipeline state on the game thread can only be safely modified after xrWaitFrame which will be unblocked by
@@ -3446,6 +3452,7 @@ void FOpenXRHMD::OnBeginSimulation_GameThread()
 
 	if (bTrackingSpaceInvalid || !ensure(PipelineState.TrackingSpace.IsValid()))
 	{
+		SCOPED_NAMED_EVENT(CreateTrackingSpace, FColor::Red);		
 		// Create the tracking space we'll use until the next recenter.
 		FTransform BaseTransform(BaseOrientation, BasePosition);
 		PipelineState.TrackingSpace = MakeShared<FTrackingSpace>(TrackingSpaceType, ToXrPose(BaseTransform, WorldToMetersScale));
