PR #12738: Ensure that game audio is automatically redirected to Meta Quest in PC VR even if the headset is not set as the default Windows audio device.

diff --git a/Engine/Plugins/Runtime/OpenXR/Source/OpenXRHMD/Private/OpenXRHMDModule.cpp b/Engine/Plugins/Runtime/OpenXR/Source/OpenXRHMD/Private/OpenXRHMDModule.cpp
--- a/Engine/Plugins/Runtime/OpenXR/Source/OpenXRHMD/Private/OpenXRHMDModule.cpp
+++ b/Engine/Plugins/Runtime/OpenXR/Source/OpenXRHMD/Private/OpenXRHMDModule.cpp
@@ -111,6 +111,7 @@ bool FOpenXRHMDModule::PreInit()
 				WCHAR DeviceGuid[XR_MAX_AUDIO_DEVICE_STR_SIZE_OCULUS];
 				GetAudioInputDeviceGuidOculus(Instance, DeviceGuid);
 				OculusAudioInputDevice = FString(XR_MAX_AUDIO_DEVICE_STR_SIZE_OCULUS, DeviceGuid);
+				OculusAudioInputDevice.TrimToNullTerminator();
 			}
 
 			PFN_xrGetAudioOutputDeviceGuidOculus GetAudioOutputDeviceGuidOculus = nullptr;
@@ -119,6 +120,7 @@ bool FOpenXRHMDModule::PreInit()
 				WCHAR DeviceGuid[XR_MAX_AUDIO_DEVICE_STR_SIZE_OCULUS];
 				GetAudioOutputDeviceGuidOculus(Instance, DeviceGuid);
 				OculusAudioOutputDevice = FString::ConstructFromPtrSize(DeviceGuid, XR_MAX_AUDIO_DEVICE_STR_SIZE_OCULUS);
+				OculusAudioOutputDevice.TrimToNullTerminator();
 			}
 		}
 		if (!CVarRetainPreInitInstance.GetValueOnAnyThread())
