OpenXRHMD: Expose layer events from for plugin extensions
-PR #12903: Expose layer events from OpenXRHMD for plugin extensions
#rnx

diff --git a/Engine/Plugins/Runtime/OpenXR/Source/OpenXRHMD/Private/OpenXRHMD.cpp b/Engine/Plugins/Runtime/OpenXR/Source/OpenXRHMD/Private/OpenXRHMD.cpp
--- a/Engine/Plugins/Runtime/OpenXR/Source/OpenXRHMD/Private/OpenXRHMD.cpp
+++ b/Engine/Plugins/Runtime/OpenXR/Source/OpenXRHMD/Private/OpenXRHMD.cpp
@@ -4491,6 +4491,38 @@ void FOpenXRHMD::UpdateLayer(FOpenXRLayer& ManagerLayer, uint32 LayerId, bool bI
 	});
 }
 
+uint32 FOpenXRHMD::CreateLayer(const FLayerDesc& InLayerDesc)
+{
+	const uint32_t LayerId = TStereoLayerManager::CreateLayer(InLayerDesc);
+
+	for (IOpenXRExtensionPlugin* Module : ExtensionPlugins)
+	{
+		Module->OnCreateLayer(LayerId);
+	}
+
+	return LayerId;
+}
+
+void FOpenXRHMD::DestroyLayer(uint32 LayerId)
+{
+	for (IOpenXRExtensionPlugin* Module : ExtensionPlugins)
+	{
+		Module->OnDestroyLayer(LayerId);
+	}
+
+	TStereoLayerManager::DestroyLayer(LayerId);
+}
+
+void FOpenXRHMD::SetLayerDesc(uint32 LayerId, const FLayerDesc& InLayerDesc)
+{
+	TStereoLayerManager::SetLayerDesc(LayerId, InLayerDesc);
+
+	for (IOpenXRExtensionPlugin* Module : ExtensionPlugins)
+	{
+		Module->OnSetLayerDesc(LayerId);
+	}
+}
+
 FOpenXRSwapchain* FOpenXRHMD::GetColorSwapchain_RenderThread()
 {
 	if (PipelinedLayerStateRendering.ColorSwapchain != nullptr)
diff --git a/Engine/Plugins/Runtime/OpenXR/Source/OpenXRHMD/Private/OpenXRHMD.h b/Engine/Plugins/Runtime/OpenXR/Source/OpenXRHMD/Private/OpenXRHMD.h
--- a/Engine/Plugins/Runtime/OpenXR/Source/OpenXRHMD/Private/OpenXRHMD.h
+++ b/Engine/Plugins/Runtime/OpenXR/Source/OpenXRHMD/Private/OpenXRHMD.h
@@ -353,6 +353,9 @@ protected:
 
 	/** TStereoLayerManager<FOpenXRLayer> */
 	void UpdateLayer(FOpenXRLayer& ManagerLayer, uint32 LayerId, bool bIsValid) override;
+	virtual uint32 CreateLayer(const FLayerDesc& InLayerDesc) override;
+	virtual void DestroyLayer(uint32 LayerId) override;
+	virtual void SetLayerDesc(uint32 LayerId, const FLayerDesc& InLayerDesc) override;
 
 	virtual bool PopulateAnalyticsAttributes(TArray<struct FAnalyticsEventAttribute>& EventAttributes) override;
 
diff --git a/Engine/Plugins/Runtime/OpenXR/Source/OpenXRHMD/Public/IOpenXRExtensionPlugin.h b/Engine/Plugins/Runtime/OpenXR/Source/OpenXRHMD/Public/IOpenXRExtensionPlugin.h
--- a/Engine/Plugins/Runtime/OpenXR/Source/OpenXRHMD/Public/IOpenXRExtensionPlugin.h
+++ b/Engine/Plugins/Runtime/OpenXR/Source/OpenXRHMD/Public/IOpenXRExtensionPlugin.h
@@ -404,6 +404,21 @@ public:
 		return InNext;
 	}	
 	
+	// FOpenXRHMD::CreateLayer
+	virtual void OnCreateLayer(uint32 LayerId)
+	{
+	}
+
+	// FOpenXRHMD::DestroyLayer
+	virtual void OnDestroyLayer(uint32 LayerId)
+	{
+	}
+
+	// FOpenXRHMD::SetLayerDesc
+	virtual void OnSetLayerDesc(uint32 LayerId)
+	{
+	}
+
 	// FOpenXRHMD::OnBeginRendering_RenderThread
 	virtual void OnSetupLayers_RenderThread(XrSession InSession, const TArray<uint32>& LayerIds)
 	{
