Fix bug in FOpenXRHMD::BuildOcclusionMesh which caused only the first quarter of an HMD's visibility mesh data to be properly uploaded to the GPU.
Additionally, ensure that the tonemap step properly clears its newly created render target if it intends to draw only to the HMD visibility mesh, in order to avoid corruption in other areas of the render target.

diff --git a/Engine/Plugins/Runtime/OpenXR/Source/OpenXRHMD/Private/OpenXRHMD.cpp b/Engine/Plugins/Runtime/OpenXR/Source/OpenXRHMD/Private/OpenXRHMD.cpp
--- a/Engine/Plugins/Runtime/OpenXR/Source/OpenXRHMD/Private/OpenXRHMD.cpp
+++ b/Engine/Plugins/Runtime/OpenXR/Source/OpenXRHMD/Private/OpenXRHMD.cpp
@@ -40,6 +40,7 @@
 #include "DefaultStereoLayers.h"
 #include "FBFoveationImageGenerator.h"
 #include "AnalyticsEventAttribute.h"
+#include "Containers/ArrayView.h"
 
 #if PLATFORM_ANDROID
 #include "Android/AndroidApplication.h"
@@ -1999,7 +2000,7 @@ bool FOpenXRHMD::BuildOcclusionMesh(XrVisibilityMaskTypeKHR Type, int View, FHMD
 
 	GetVisibilityMaskKHR(Session, SelectedViewConfigurationType, View, Type, &VisibilityMask);
 
-	IndexInitialData.WriteData(OutIndices.Get(), VisibilityMask.indexCountOutput);
+	IndexInitialData.WriteArray(MakeConstArrayView(OutIndices.Get(), VisibilityMask.indexCountOutput));
 
 	// We need to apply the eye's projection matrix to each vertex
 	FMatrix Projection = GetStereoProjectionMatrix(View);
diff --git a/Engine/Source/Runtime/Renderer/Private/PostProcess/PostProcessTonemap.cpp b/Engine/Source/Runtime/Renderer/Private/PostProcess/PostProcessTonemap.cpp
--- a/Engine/Source/Runtime/Renderer/Private/PostProcess/PostProcessTonemap.cpp
+++ b/Engine/Source/Runtime/Renderer/Private/PostProcess/PostProcessTonemap.cpp
@@ -643,7 +643,7 @@ FScreenPassTexture AddTonemapPass(FRDGBuilder& GraphBuilder, const FViewInfo& Vi
 		Output = FScreenPassRenderTarget(
 			GraphBuilder.CreateTexture(OutputDesc, TEXT("Tonemap")),
 			Inputs.SceneColor.ViewRect,
-			ERenderTargetLoadAction::ENoAction);
+			View.GetOverwriteLoadAction());
 	}
 
 	const bool bShouldWriteAlphaChannel = ShouldWriteAlphaChannel(View, Inputs, Output.Texture->Desc.Format);
