PR #13003: Enable depth testing for stereo layers

diff --git a/Engine/Plugins/Runtime/OpenXR/Source/OpenXRHMD/Private/OpenXRHMD.cpp b/Engine/Plugins/Runtime/OpenXR/Source/OpenXRHMD/Private/OpenXRHMD.cpp
--- a/Engine/Plugins/Runtime/OpenXR/Source/OpenXRHMD/Private/OpenXRHMD.cpp
+++ b/Engine/Plugins/Runtime/OpenXR/Source/OpenXRHMD/Private/OpenXRHMD.cpp
@@ -2996,6 +2996,7 @@ void FOpenXRHMD::SetupFrameLayers_RenderThread(FRHICommandListImmediate& RHICmdL
 
 	PipelinedLayerStateRendering.NativeOverlays.Reset(NativeLayers.Num());
 	PipelinedLayerStateRendering.NativeOverlaySwapchains.Reset(NativeLayers.Num());
+	PipelinedLayerStateRendering.CompositionDepthTestLayers.Reset(NativeLayers.Num());
 
 	// Set up our OpenXR info per native layer. Emulated layers have everything in FLayerDesc.
 	for (const FOpenXRLayer& Layer : NativeLayers)
@@ -3007,7 +3008,11 @@ void FOpenXRHMD::SetupFrameLayers_RenderThread(FRHICommandListImmediate& RHICmdL
 		
 		EOpenXRLayerCreationFlags LayerCreationFlags = EOpenXRLayerCreationFlags::None;
 		LayerCreationFlags |= IsExtensionEnabled(XR_KHR_COMPOSITION_LAYER_EQUIRECT2_EXTENSION_NAME) ? EOpenXRLayerCreationFlags::EquirectLayer2Supported : EOpenXRLayerCreationFlags::None;
+		LayerCreationFlags |= IsExtensionEnabled(XR_FB_COMPOSITION_LAYER_DEPTH_TEST_EXTENSION_NAME) ? EOpenXRLayerCreationFlags::DepthTestSupported : EOpenXRLayerCreationFlags::None;
+
 		TArray<FXrCompositionLayerUnion> Headers = Layer.CreateOpenXRLayer(InvTrackingToWorld, WorldToMeters, Space, LayerCreationFlags);
+		Layer.ApplyCompositionDepthTestLayer(Headers, LayerCreationFlags, PipelinedLayerStateRendering.CompositionDepthTestLayers);
+
 		PipelinedLayerStateRendering.NativeOverlays.Append(Headers);
 		UpdateLayerSwapchainTexture(Layer, RHICmdList);
 	}
@@ -3839,6 +3844,7 @@ void FOpenXRHMD::OnFinishRendering_RHIThread(IRHICommandContext& RHICmdContext)
 		XrCompositionLayerProjection Layer = {};
 		XrCompositionLayerAlphaBlendFB LayerAlphaBlend = { XR_TYPE_COMPOSITION_LAYER_ALPHA_BLEND_FB };
 		XrCompositionLayerColorScaleBiasKHR ColorScaleBias = { XR_TYPE_COMPOSITION_LAYER_COLOR_SCALE_BIAS_KHR };
+		XrCompositionLayerDepthTestFB LayerDepthTest = { XR_TYPE_COMPOSITION_LAYER_DEPTH_TEST_FB };
 		if (EnumHasAnyFlags(PipelinedLayerStateRHI.LayerStateFlags, EOpenXRLayerStateFlags::SubmitBackgroundLayer))
 		{
 			Layer.type = XR_TYPE_COMPOSITION_LAYER_PROJECTION;
@@ -3877,6 +3883,15 @@ void FOpenXRHMD::OnFinishRendering_RHIThread(IRHICommandContext& RHICmdContext)
 				Layer.next = &ColorScaleBias;
 			}
 
+			if (IsExtensionEnabled(XR_FB_COMPOSITION_LAYER_DEPTH_TEST_EXTENSION_NAME) && 
+				EnumHasAnyFlags(PipelinedLayerStateRendering.LayerStateFlags, EOpenXRLayerStateFlags::SubmitDepthLayer))
+			{
+				LayerDepthTest.next = const_cast<void*>(Layer.next);
+				LayerDepthTest.depthMask = true;
+				LayerDepthTest.compareOp = XR_COMPARE_OP_LESS_FB;
+				Layer.next = &LayerDepthTest;
+			}
+
 			for (IOpenXRExtensionPlugin* Module : ExtensionPlugins)
 			{
 				Layer.next = Module->OnEndProjectionLayer(Session, 0, Layer.next, Layer.layerFlags);
diff --git a/Engine/Plugins/Runtime/OpenXR/Source/OpenXRHMD/Private/OpenXRHMD.h b/Engine/Plugins/Runtime/OpenXR/Source/OpenXRHMD/Private/OpenXRHMD.h
--- a/Engine/Plugins/Runtime/OpenXR/Source/OpenXRHMD/Private/OpenXRHMD.h
+++ b/Engine/Plugins/Runtime/OpenXR/Source/OpenXRHMD/Private/OpenXRHMD.h
@@ -203,6 +203,7 @@ public:
 		TArray<FXrCompositionLayerUnion> NativeOverlays;
 		TArray<XrCompositionLayerProjectionView> ProjectionLayers;
 		TArray<XrCompositionLayerDepthInfoKHR> DepthLayers;
+		TArray<XrCompositionLayerDepthTestFB> CompositionDepthTestLayers;
 
 		TArray<XrSwapchainSubImage> ColorImages;
 		TArray<XrSwapchainSubImage> DepthImages;
diff --git a/Engine/Plugins/Runtime/OpenXR/Source/OpenXRHMD/Private/OpenXRHMDModule.cpp b/Engine/Plugins/Runtime/OpenXR/Source/OpenXRHMD/Private/OpenXRHMDModule.cpp
--- a/Engine/Plugins/Runtime/OpenXR/Source/OpenXRHMD/Private/OpenXRHMDModule.cpp
+++ b/Engine/Plugins/Runtime/OpenXR/Source/OpenXRHMD/Private/OpenXRHMDModule.cpp
@@ -547,6 +547,7 @@ bool FOpenXRHMDModule::GetOptionalExtensions(TArray<const ANSICHAR*>& OutExtensi
 	OutExtensions.Add(XR_FB_FOVEATION_VULKAN_EXTENSION_NAME);
 #endif
 	OutExtensions.Add(XR_KHR_COMPOSITION_LAYER_DEPTH_EXTENSION_NAME);
+	OutExtensions.Add(XR_FB_COMPOSITION_LAYER_DEPTH_TEST_EXTENSION_NAME);
 	OutExtensions.Add(XR_KHR_COMPOSITION_LAYER_CYLINDER_EXTENSION_NAME);
 	OutExtensions.Add(XR_KHR_COMPOSITION_LAYER_EQUIRECT_EXTENSION_NAME);
 	OutExtensions.Add(XR_KHR_COMPOSITION_LAYER_EQUIRECT2_EXTENSION_NAME);
diff --git a/Engine/Plugins/Runtime/OpenXR/Source/OpenXRHMD/Private/OpenXRHMD_Layer.cpp b/Engine/Plugins/Runtime/OpenXR/Source/OpenXRHMD/Private/OpenXRHMD_Layer.cpp
--- a/Engine/Plugins/Runtime/OpenXR/Source/OpenXRHMD/Private/OpenXRHMD_Layer.cpp
+++ b/Engine/Plugins/Runtime/OpenXR/Source/OpenXRHMD/Private/OpenXRHMD_Layer.cpp
@@ -111,6 +111,26 @@ TArray<FXrCompositionLayerUnion> FOpenXRLayer::CreateOpenXRLayer(FTransform InvT
 	return Headers;
 }
 
+void FOpenXRLayer::ApplyCompositionDepthTestLayer(TArray<FXrCompositionLayerUnion>& Headers, EOpenXRLayerCreationFlags LayerCreationFlags, TArray<XrCompositionLayerDepthTestFB>& InCompositionDepthTestLayers) const
+{
+	const bool bUseDepthTest = EnumHasAllFlags((IStereoLayers::ELayerFlags)Desc.Flags, IStereoLayers::LAYER_FLAG_SUPPORT_DEPTH) &&
+		EnumHasAllFlags(LayerCreationFlags, EOpenXRLayerCreationFlags::DepthTestSupported);
+
+	if (bUseDepthTest)
+	{
+		for (FXrCompositionLayerUnion& Header : Headers)
+		{
+			InCompositionDepthTestLayers.AddUninitialized();
+			XrCompositionLayerDepthTestFB& LayerDepthTest = InCompositionDepthTestLayers.Last();
+			LayerDepthTest.type = XR_TYPE_COMPOSITION_LAYER_DEPTH_TEST_FB;
+			LayerDepthTest.next = Header.Header.next;
+			LayerDepthTest.depthMask = true;
+			LayerDepthTest.compareOp = XR_COMPARE_OP_LESS_FB;
+			Header.Header.next = &LayerDepthTest;
+		}
+	}
+}
+
 void FOpenXRLayer::CreateOpenXRCylinderLayer(bool bIsStereo, bool bNoAlpha, FTransform PositionTransform, float WorldToMeters, XrSpace Space, TArray<FXrCompositionLayerUnion>& Headers) const
 {
 	XrCompositionLayerCylinderKHR Cylinder = { XR_TYPE_COMPOSITION_LAYER_CYLINDER_KHR, /*next*/ nullptr };
diff --git a/Engine/Plugins/Runtime/OpenXR/Source/OpenXRHMD/Private/OpenXRHMD_Layer.h b/Engine/Plugins/Runtime/OpenXR/Source/OpenXRHMD/Private/OpenXRHMD_Layer.h
--- a/Engine/Plugins/Runtime/OpenXR/Source/OpenXRHMD/Private/OpenXRHMD_Layer.h
+++ b/Engine/Plugins/Runtime/OpenXR/Source/OpenXRHMD/Private/OpenXRHMD_Layer.h
@@ -50,6 +50,8 @@ struct FOpenXRLayer
 	FVector2D GetLeftQuadSize() const;
 
 	TArray<FXrCompositionLayerUnion> CreateOpenXRLayer(FTransform InvTrackingToWorld, float WorldToMeters, XrSpace Space, EOpenXRLayerCreationFlags CreationFlags) const;
+
+	void ApplyCompositionDepthTestLayer(TArray<FXrCompositionLayerUnion>& Headers, EOpenXRLayerCreationFlags CreationFlags, TArray<XrCompositionLayerDepthTestFB>& InCompositionDepthTestLayers) const;
 	
 private:
 	void CreateOpenXRQuadLayer(bool bIsStereo, bool bNoAlpha, FTransform PositionTransform, float WorldToMeters, XrSpace Space, TArray<FXrCompositionLayerUnion>& Headers) const;
@@ -57,7 +59,6 @@ private:
 	void CreateOpenXREquirectLayer(bool bIsStereo, bool bNoAlpha, FTransform PositionTransform, float WorldToMeters, XrSpace Space, TArray<FXrCompositionLayerUnion>& Headers) const;
 	void CreateOpenXREquirect2Layer(bool bIsStereo, bool bNoAlpha, FTransform PositionTransform, float WorldToMeters, XrSpace Space, TArray<FXrCompositionLayerUnion>& Headers) const;
 	void SetupEquirect2(FVector2D Scale, FVector2D Bias, FVector2D Position, FVector2D Size, FTransform PositionTransform, float WorldToMeters, XrCompositionLayerEquirect2KHR& Equirect2) const;
-
 };
 
 bool GetLayerDescMember(const FOpenXRLayer& Layer, IStereoLayers::FLayerDesc& OutLayerDesc);
diff --git a/Engine/Plugins/Runtime/OpenXR/Source/OpenXRHMD/Public/OpenXRCore.h b/Engine/Plugins/Runtime/OpenXR/Source/OpenXRHMD/Public/OpenXRCore.h
--- a/Engine/Plugins/Runtime/OpenXR/Source/OpenXRHMD/Public/OpenXRCore.h
+++ b/Engine/Plugins/Runtime/OpenXR/Source/OpenXRHMD/Public/OpenXRCore.h
@@ -149,6 +149,7 @@ enum class EOpenXRLayerCreationFlags : uint32
 {
 	None = 0u,
 	EquirectLayer2Supported = (1u << 0),
+	DepthTestSupported = (1u << 1),
 };
 
 ENUM_CLASS_FLAGS(EOpenXRLayerCreationFlags);
