OpenXR need to account for the fact that the LOCAL_FLOOR space may not exist while using openxr 1.1
-We were requring that LOCAL_FLOOR space be supported, which is incorrect.  Now it is optional.
-Additionally implemented the grip_surface pose which was also triggering an ensure.

diff --git a/Engine/Platforms/VisionOS/Plugins/Runtime/OpenXRVisionOS/Source/OXRVisionOS/Private/OXRVisionOSController.cpp b/Engine/Platforms/VisionOS/Plugins/Runtime/OpenXRVisionOS/Source/OXRVisionOS/Private/OXRVisionOSController.cpp
--- a/Engine/Platforms/VisionOS/Plugins/Runtime/OpenXRVisionOS/Source/OXRVisionOS/Private/OXRVisionOSController.cpp
+++ b/Engine/Platforms/VisionOS/Plugins/Runtime/OpenXRVisionOS/Source/OXRVisionOS/Private/OXRVisionOSController.cpp
@@ -19,8 +19,10 @@ namespace OXRVisionOS
 
 		TEXT("GripL"),
 		TEXT("AimL"),
+		TEXT("GripSurfaceL"),
 		TEXT("GripR"),
 		TEXT("AimR"),
+		TEXT("GripSurfaceR"),
 		TEXT("HMDPose"),
 
 //		TEXT("GazePose"),
diff --git a/Engine/Platforms/VisionOS/Plugins/Runtime/OpenXRVisionOS/Source/OXRVisionOS/Private/OXRVisionOSController.h b/Engine/Platforms/VisionOS/Plugins/Runtime/OpenXRVisionOS/Source/OXRVisionOS/Private/OXRVisionOSController.h
--- a/Engine/Platforms/VisionOS/Plugins/Runtime/OpenXRVisionOS/Source/OXRVisionOS/Private/OXRVisionOSController.h
+++ b/Engine/Platforms/VisionOS/Plugins/Runtime/OpenXRVisionOS/Source/OXRVisionOS/Private/OXRVisionOSController.h
@@ -29,8 +29,10 @@ enum class EOXRVisionOSControllerButton : int32
 	/** Pose inputs, handled by OXRVisionOSTracker **/
 	GripL = TotalButtonCount,
 	AimL,
+	GripSurfaceL,
 	GripR,
 	AimR,
+	GripSurfaceR,
 	HMDPose,
 
 	/** The eye tracking input **/
diff --git a/Engine/Platforms/VisionOS/Plugins/Runtime/OpenXRVisionOS/Source/OXRVisionOS/Private/OXRVisionOSInstance.cpp b/Engine/Platforms/VisionOS/Plugins/Runtime/OpenXRVisionOS/Source/OXRVisionOS/Private/OXRVisionOSInstance.cpp
--- a/Engine/Platforms/VisionOS/Plugins/Runtime/OpenXRVisionOS/Source/OXRVisionOS/Private/OXRVisionOSInstance.cpp
+++ b/Engine/Platforms/VisionOS/Plugins/Runtime/OpenXRVisionOS/Source/OXRVisionOS/Private/OXRVisionOSInstance.cpp
@@ -579,6 +579,8 @@ void FOXRVisionOSInstance::CreateInteractionProfiles()
 		Pairs.Add(StringToPath("/user/hand/right/input/grip/pose"), EOXRVisionOSControllerButton::GripR);
 		Pairs.Add(StringToPath("/user/hand/left/input/aim/pose"), EOXRVisionOSControllerButton::AimL);
 		Pairs.Add(StringToPath("/user/hand/right/input/aim/pose"), EOXRVisionOSControllerButton::AimR);
+		Pairs.Add(StringToPath("/user/hand/left/input/grip_surface/pose"), EOXRVisionOSControllerButton::GripSurfaceL);
+		Pairs.Add(StringToPath("/user/hand/right/input/grip_surface/pose"), EOXRVisionOSControllerButton::GripSurfaceR);
 
 		// Parent Paths  One can use the 'parent' path of a value or click and it should behave like click if there is one and value if there is not.  OpenXR spec 11.4.
 		Pairs.Add(StringToPath("/user/hand/left/input/select"), EOXRVisionOSControllerButton::PushL);			// stick click
diff --git a/Engine/Platforms/VisionOS/Plugins/Runtime/OpenXRVisionOS/Source/OXRVisionOS/Private/OXRVisionOSSession.cpp b/Engine/Platforms/VisionOS/Plugins/Runtime/OpenXRVisionOS/Source/OXRVisionOS/Private/OXRVisionOSSession.cpp
--- a/Engine/Platforms/VisionOS/Plugins/Runtime/OpenXRVisionOS/Source/OXRVisionOS/Private/OXRVisionOSSession.cpp
+++ b/Engine/Platforms/VisionOS/Plugins/Runtime/OpenXRVisionOS/Source/OXRVisionOS/Private/OXRVisionOSSession.cpp
@@ -104,6 +104,7 @@ FOXRVisionOSSession::FOXRVisionOSSession(const XrSessionCreateInfo* createInfo,
 	const UOXRVisionOSRuntimeSettings* Settings = GetDefault<UOXRVisionOSRuntimeSettings>();
 	AimPoseAdjustment = FQuat(Settings->OXRVisionOSAimPoseAdjustment);
 	GripPoseAdjustment = FQuat(Settings->OXRVisionOSGripPoseAdjustment);
+	GripSurfacePoseAdjustment = FQuat(Settings->OXRVisionOSGripSurfacePoseAdjustment);
 
 	static_assert(OXRVisionOSSessionHelpers::ELayerIndex::Num <= OXRVisionOSSessionHelpers::MaxOXRVisionOSLayers);
 
@@ -1410,10 +1411,12 @@ void FOXRVisionOSSession::GetTransformForButton(EOXRVisionOSControllerButton But
 	{
 	case EOXRVisionOSControllerButton::GripL:
 	case EOXRVisionOSControllerButton::AimL:
+	case EOXRVisionOSControllerButton::GripSurfaceL:
 		DeviceHandle = Controllers->GetLeftControllerHandle();
 		break;
 	case EOXRVisionOSControllerButton::GripR:
 	case EOXRVisionOSControllerButton::AimR:
+	case EOXRVisionOSControllerButton::GripSurfaceR:
 		DeviceHandle = Controllers->GetRightControllerHandle();
 		break;
 	case EOXRVisionOSControllerButton::HMDPose:
@@ -1484,6 +1487,10 @@ void FOXRVisionOSSession::GetTransformForButton(EOXRVisionOSControllerButton But
 			{
 				OutTransform.SetRotation(OutTransform.TransformRotation(AimPoseAdjustment));
 			}
+			else if (Button == EOXRVisionOSControllerButton::GripSurfaceL || Button == EOXRVisionOSControllerButton::GripSurfaceR)
+			{
+				OutTransform.SetRotation(OutTransform.TransformRotation(GripSurfacePoseAdjustment));
+			}
 		}
 	}
 }
diff --git a/Engine/Platforms/VisionOS/Plugins/Runtime/OpenXRVisionOS/Source/OXRVisionOS/Private/OXRVisionOSSession.h b/Engine/Platforms/VisionOS/Plugins/Runtime/OpenXRVisionOS/Source/OXRVisionOS/Private/OXRVisionOSSession.h
--- a/Engine/Platforms/VisionOS/Plugins/Runtime/OpenXRVisionOS/Source/OXRVisionOS/Private/OXRVisionOSSession.h
+++ b/Engine/Platforms/VisionOS/Plugins/Runtime/OpenXRVisionOS/Source/OXRVisionOS/Private/OXRVisionOSSession.h
@@ -215,6 +215,7 @@ private:
 	TArray<TSharedPtr<FOXRVisionOSSpace, ESPMode::ThreadSafe>> Spaces;
 	FQuat GripPoseAdjustment;
 	FQuat AimPoseAdjustment;
+	FQuat GripSurfacePoseAdjustment;
 
 	FPlatformMemory::FPlatformVirtualMemoryBlock m_reprojectionData;
 	FPlatformMemory::FPlatformVirtualMemoryBlock m_displayData;
diff --git a/Engine/Platforms/VisionOS/Plugins/Runtime/OpenXRVisionOS/Source/OXRVisionOS/Private/OXRVisionOSTracker.cpp b/Engine/Platforms/VisionOS/Plugins/Runtime/OpenXRVisionOS/Source/OXRVisionOS/Private/OXRVisionOSTracker.cpp
--- a/Engine/Platforms/VisionOS/Plugins/Runtime/OpenXRVisionOS/Source/OXRVisionOS/Private/OXRVisionOSTracker.cpp
+++ b/Engine/Platforms/VisionOS/Plugins/Runtime/OpenXRVisionOS/Source/OXRVisionOS/Private/OXRVisionOSTracker.cpp
@@ -325,6 +325,7 @@ void FOXRVisionOSTracker::GetActionStatePose(EOXRVisionOSControllerButton Button
 	{
 	case EOXRVisionOSControllerButton::GripL:
 	case EOXRVisionOSControllerButton::AimL:
+	case EOXRVisionOSControllerButton::GripSurfaceL:
 	{
 		FDeviceInfo* DeviceInfo = DeviceInfoMap.Find(TrackerDeviceType::DEVICE_VR_CONTROLLER_LEFT);
 		OutActive = DeviceInfo->bIsActive; //TODO is this correct?
@@ -332,6 +333,7 @@ void FOXRVisionOSTracker::GetActionStatePose(EOXRVisionOSControllerButton Button
 	}
 	case EOXRVisionOSControllerButton::GripR:
 	case EOXRVisionOSControllerButton::AimR:
+	case EOXRVisionOSControllerButton::GripSurfaceR:
 	{
 		FDeviceInfo* DeviceInfo = DeviceInfoMap.Find(TrackerDeviceType::DEVICE_VR_CONTROLLER_RIGHT);
 		OutActive = DeviceInfo->bIsActive; //TODO is this correct?
diff --git a/Engine/Platforms/VisionOS/Plugins/Runtime/OpenXRVisionOS/Source/OXRVisionOSSettings/Public/OXRVisionOSRuntimeSettings.h b/Engine/Platforms/VisionOS/Plugins/Runtime/OpenXRVisionOS/Source/OXRVisionOSSettings/Public/OXRVisionOSRuntimeSettings.h
--- a/Engine/Platforms/VisionOS/Plugins/Runtime/OpenXRVisionOS/Source/OXRVisionOSSettings/Public/OXRVisionOSRuntimeSettings.h
+++ b/Engine/Platforms/VisionOS/Plugins/Runtime/OpenXRVisionOS/Source/OXRVisionOSSettings/Public/OXRVisionOSRuntimeSettings.h
@@ -30,6 +30,10 @@ public:
 	UPROPERTY(config, EditAnywhere, Category = "OpenXR visionOS", Meta = (DisplayName = "Grip Pose Adjustment"))
 	FRotator OXRVisionOSGripPoseAdjustment = FRotator(45.0f,0.0f,0.0f);
 
+	// Defines a rotation adjustment to the grip surface pose from the OpenXR visionOS native controller orientation.  We provide a default adjustment to grip surface.
+	UPROPERTY(config, EditAnywhere, Category = "OpenXR visionOS", Meta = (DisplayName = "Grip Surface Pose Adjustment"))
+	FRotator OXRVisionOSGripSurfacePoseAdjustment = FRotator(0.0f, 0.0f, 0.0f);
+
 	// Number of UI2D/quad layers supported by OpenXR visionOS, required at app startup time
 	UPROPERTY(config, EditAnywhere, Category = "OpenXR visionOS", Meta = (DeprecatedProperty, ClampMin = "0", ClampMax = "4", DisplayName = "Max Number UI Layers"))
 	int32 OXRVisionOSNumUILayers = 0;
diff --git a/Engine/Plugins/Runtime/OpenXR/Source/OpenXRHMD/Private/OpenXRHMD.cpp b/Engine/Plugins/Runtime/OpenXR/Source/OpenXRHMD/Private/OpenXRHMD.cpp
--- a/Engine/Plugins/Runtime/OpenXR/Source/OpenXRHMD/Private/OpenXRHMD.cpp
+++ b/Engine/Plugins/Runtime/OpenXR/Source/OpenXRHMD/Private/OpenXRHMD.cpp
@@ -687,6 +687,8 @@ bool FOpenXRHMD::GetTrackingOriginTransform(TEnumAsByte<EHMDTrackingOrigin::Type
 		}
 		break;
 	case EHMDTrackingOrigin::LocalFloor:
+		// This fallback logic probably should not exist, but changing it could break exisitng project.  
+		// If we do a more comprehensive refactor of these api's we may want to eliminate this.
 		Space = bLocalFloorSpaceSupported? LocalFloorSpace : LocalSpace;
 		break;
 	case EHMDTrackingOrigin::Stage:
@@ -994,13 +996,17 @@ void FOpenXRHMD::Recenter(EOrientPositionSelector::Type Selector, float Yaw)
 	XrSpaceLocation DeviceLocation = { XR_TYPE_SPACE_LOCATION, nullptr };
 
 	XrSpace BaseSpace = XR_NULL_HANDLE;
-	if (bLocalFloorSpaceSupported && TrackingSpaceType == XR_REFERENCE_SPACE_TYPE_LOCAL_FLOOR)
+	if (TrackingSpaceType == XR_REFERENCE_SPACE_TYPE_LOCAL_FLOOR)
 	{
 		BaseSpace = LocalFloorSpace;
 	}
+	else if (TrackingSpaceType == XR_REFERENCE_SPACE_TYPE_STAGE)
+	{
+		BaseSpace = StageSpace;
+	}
 	else
 	{
-		BaseSpace = TrackingSpaceType == XR_REFERENCE_SPACE_TYPE_STAGE ? StageSpace : LocalSpace;
+		BaseSpace = LocalSpace;
 	}
 	if (bUseCustomReferenceSpace)
 	{
@@ -1097,15 +1103,15 @@ void FOpenXRHMD::SetTrackingOrigin(EHMDTrackingOrigin::Type NewOrigin)
 		return;
 	}
 
-	if (NewOrigin == EHMDTrackingOrigin::LocalFloor && bLocalFloorSpaceSupported)
+	if (NewOrigin == EHMDTrackingOrigin::LocalFloor && (LocalFloorSpace != nullptr))
 	{
 		TrackingSpaceType = XR_REFERENCE_SPACE_TYPE_LOCAL_FLOOR;
 	}
 	else if (NewOrigin == EHMDTrackingOrigin::Local)
 	{
-		TrackingSpaceType = XR_REFERENCE_SPACE_TYPE_LOCAL;  // Local space is always supported
+		TrackingSpaceType = XR_REFERENCE_SPACE_TYPE_LOCAL;  // Local space is always supported, but we only prefer it if it was requested.
 	}
-	else if (StageSpace) // Either stage is requested, or floor was requested but floor is not supported.
+	else if (StageSpace != nullptr) // Either stage is requested, or floor was requested but floor is not supported (stage meets the requirements for floor, and more).
 	{
 		TrackingSpaceType = XR_REFERENCE_SPACE_TYPE_STAGE;
 	}
@@ -2233,9 +2239,8 @@ bool FOpenXRHMD::OnStereoStartup()
 	SpaceInfo.referenceSpaceType = XR_REFERENCE_SPACE_TYPE_LOCAL;
 	XR_ENSURE(xrCreateReferenceSpace(Session, &SpaceInfo, &LocalSpace));
 
-	if(bLocalFloorSpaceSupported)
+	if(ReferenceSpaces.Contains(XR_REFERENCE_SPACE_TYPE_LOCAL_FLOOR))
 	{
-		ensure(ReferenceSpaces.Contains(XR_REFERENCE_SPACE_TYPE_LOCAL_FLOOR));
 		SpaceInfo.referenceSpaceType = XR_REFERENCE_SPACE_TYPE_LOCAL_FLOOR;
 		XR_ENSURE(xrCreateReferenceSpace(Session, &SpaceInfo, &LocalFloorSpace));
 	}
@@ -2269,9 +2274,8 @@ bool FOpenXRHMD::OnStereoStartup()
 	{
 		TrackingSpaceType = XR_REFERENCE_SPACE_TYPE_STAGE;
 	}
-	else if (bLocalFloorSpaceSupported)
+	else if (ReferenceSpaces.Contains(XR_REFERENCE_SPACE_TYPE_LOCAL_FLOOR))
 	{
-		ensure(ReferenceSpaces.Contains(XR_REFERENCE_SPACE_TYPE_LOCAL_FLOOR));
 		TrackingSpaceType = XR_REFERENCE_SPACE_TYPE_LOCAL_FLOOR;
 	}
 	else
