Add TRACE_BOOKMARK tags to major OpenXR frame events (xrWaitFrame, xrBeginFrame, xrEndFrame) to make the frame easier to follow in UnrealInsights.

diff --git a/Engine/Plugins/Runtime/OpenXR/Source/OpenXRHMD/Private/OpenXRHMD.cpp b/Engine/Plugins/Runtime/OpenXR/Source/OpenXRHMD/Private/OpenXRHMD.cpp
--- a/Engine/Plugins/Runtime/OpenXR/Source/OpenXRHMD/Private/OpenXRHMD.cpp
+++ b/Engine/Plugins/Runtime/OpenXR/Source/OpenXRHMD/Private/OpenXRHMD.cpp
@@ -3535,6 +3535,7 @@ void FOpenXRHMD::OnBeginSimulation_GameThread()
 	UE_LOG(LogHMD, VeryVerbose, TEXT("%s WF_%i xrWaitFrame Calling..."), HMDThreadString(), WaitCount);
 	{
 		SCOPED_NAMED_EVENT(xrWaitFrame, FColor::Red);
+		TRACE_BOOKMARK(TEXT("xrWaitFrame: %d"), WaitCount);
 		XR_ENSURE(xrWaitFrame(Session, &WaitInfo, &FrameState));
 	}
 	UE_LOG(LogHMD, VeryVerbose, TEXT("%s WF_%i xrWaitFrame Complete"), HMDThreadString(), WaitCount);
@@ -3871,6 +3872,8 @@ void FOpenXRHMD::OnBeginRendering_RHIThread(IRHICommandContext& RHICmdContext, c
 	}
 	static int BeginCount = 0;
 	PipelinedFrameStateRHI.BeginCount = ++BeginCount;
+
+	TRACE_BOOKMARK(TEXT("xrBeginFrame: %d"), BeginCount);
 	UE_LOG(LogHMD, VeryVerbose, TEXT("%s WF_%i xrBeginFrame BeginCount: %i"), HMDThreadString(), PipelinedFrameStateRHI.WaitCount, PipelinedFrameStateRHI.BeginCount);
 	XrResult Result = xrBeginFrame(Session, &BeginInfo);
 	if (XR_SUCCEEDED(Result))
@@ -4068,6 +4071,8 @@ void FOpenXRHMD::OnFinishRendering_RHIThread(IRHICommandContext& RHICmdContext)
 #endif
 		static int EndCount = 0;
 		PipelinedFrameStateRHI.EndCount = ++EndCount;
+
+		TRACE_BOOKMARK(TEXT("xrEndFrame: %d"), EndCount);
 		UE_LOG(LogHMD, VeryVerbose, TEXT("%s WF_%i xrEndFrame WaitCount: %i BeginCount: %i EndCount: %i"), HMDThreadString(), PipelinedFrameStateRHI.WaitCount, PipelinedFrameStateRHI.WaitCount, PipelinedFrameStateRHI.BeginCount, PipelinedFrameStateRHI.EndCount);
 		XR_ENSURE(xrEndFrame(Session, &EndInfo));
 	}
