Support generic passthrough on standalone Meta Quest mobile.
- Add use-feature com.oculus.feature.PASSTHROUGH to the Android manifest when packaging for Meta Quest.
- Avoid pulling in the XR_FB_composition_layer_alpha_blend OpenXR extension when XR_EXT_composition_layer_inverted_alpha is available. The proprietary FB extension restricts available environment blend modes and prevents using passthrough via the generic EXT extension.
- Set bProjectionLayerAlphaEnabled to true in OpenXRHMD when we are on a mobile platform and r.Mobile.PropagateAlpha=True.

diff --git a/Engine/Plugins/Runtime/OpenXR/Source/OpenXRHMD/Private/OpenXRHMD.cpp b/Engine/Plugins/Runtime/OpenXR/Source/OpenXRHMD/Private/OpenXRHMD.cpp
--- a/Engine/Plugins/Runtime/OpenXR/Source/OpenXRHMD/Private/OpenXRHMD.cpp
+++ b/Engine/Plugins/Runtime/OpenXR/Source/OpenXRHMD/Private/OpenXRHMD.cpp
@@ -1740,7 +1740,8 @@ bool FOpenXRHMD::ReconfigureForShaderPlatform(EShaderPlatform NewShaderPlatform)
 	bIsMobileMultiViewEnabled = Aspects.IsMobileMultiViewEnabled();
 
 	static const auto CVarPropagateAlpha = IConsoleManager::Get().FindConsoleVariable(TEXT("r.PostProcessing.PropagateAlpha"));
-	bProjectionLayerAlphaEnabled = !IsMobilePlatform(NewShaderPlatform) && CVarPropagateAlpha->GetBool();
+	static const auto CVarPropagateAlphaMobile = IConsoleManager::Get().FindConsoleVariable(TEXT("r.Mobile.PropagateAlpha"));
+	bProjectionLayerAlphaEnabled = IsMobilePlatform(NewShaderPlatform) ? CVarPropagateAlphaMobile->GetBool() : CVarPropagateAlpha->GetBool();
 
 	ConfiguredShaderPlatform = NewShaderPlatform;
 
diff --git a/Engine/Plugins/Runtime/OpenXR/Source/OpenXRHMD/Private/OpenXRHMDModule.cpp b/Engine/Plugins/Runtime/OpenXR/Source/OpenXRHMD/Private/OpenXRHMDModule.cpp
--- a/Engine/Plugins/Runtime/OpenXR/Source/OpenXRHMD/Private/OpenXRHMDModule.cpp
+++ b/Engine/Plugins/Runtime/OpenXR/Source/OpenXRHMD/Private/OpenXRHMDModule.cpp
@@ -559,8 +559,6 @@ bool FOpenXRHMDModule::GetOptionalExtensions(TArray<const ANSICHAR*>& OutExtensi
 	OutExtensions.Add(XR_EXT_DPAD_BINDING_EXTENSION_NAME);
 	OutExtensions.Add(XR_EXT_PALM_POSE_EXTENSION_NAME);
 	OutExtensions.Add(XR_EXT_ACTIVE_ACTION_SET_PRIORITY_EXTENSION_NAME);
-	OutExtensions.Add(XR_FB_COMPOSITION_LAYER_ALPHA_BLEND_EXTENSION_NAME);
-	OutExtensions.Add(XR_EXT_COMPOSITION_LAYER_INVERTED_ALPHA_EXTENSION_NAME);
 	OutExtensions.Add(XR_FB_FOVEATION_EXTENSION_NAME);
 	OutExtensions.Add(XR_FB_SWAPCHAIN_UPDATE_STATE_EXTENSION_NAME);
 	OutExtensions.Add(XR_FB_FOVEATION_CONFIGURATION_EXTENSION_NAME); 
@@ -568,6 +566,17 @@ bool FOpenXRHMDModule::GetOptionalExtensions(TArray<const ANSICHAR*>& OutExtensi
 	OutExtensions.Add(XR_OCULUS_AUDIO_DEVICE_GUID_EXTENSION_NAME);
 #endif
 
+	// These extensions are mutually incompatible because XR_FB_COMPOSITION_LAYER_ALPHA_BLEND disables non-opaque environment blend modes
+	// XR_EXT_COMPOSITION_LAYER_INVERTED_ALPHA is the more modern and platform-agnostic option and is preferred
+	if (AvailableExtensions.Contains(XR_EXT_COMPOSITION_LAYER_INVERTED_ALPHA_EXTENSION_NAME))
+	{
+		OutExtensions.Add(XR_EXT_COMPOSITION_LAYER_INVERTED_ALPHA_EXTENSION_NAME);
+	}
+	else
+	{
+		OutExtensions.Add(XR_FB_COMPOSITION_LAYER_ALPHA_BLEND_EXTENSION_NAME);
+	}
+	
 	OutExtensions.Add(XR_EXT_LOCAL_FLOOR_EXTENSION_NAME);
 	return true;
 }
diff --git a/Engine/Source/Programs/UnrealBuildTool/Platform/Android/UEDeployAndroid.cs b/Engine/Source/Programs/UnrealBuildTool/Platform/Android/UEDeployAndroid.cs
--- a/Engine/Source/Programs/UnrealBuildTool/Platform/Android/UEDeployAndroid.cs
+++ b/Engine/Source/Programs/UnrealBuildTool/Platform/Android/UEDeployAndroid.cs
@@ -3119,6 +3119,9 @@ namespace UnrealBuildTool
 					// need for hand tracking on MetaQuest
 					Text.AppendLine("\t<uses-feature android:name=\"oculus.software.handtracking\" android:required=\"false\"/>");
 					Text.AppendLine("\t<uses-permission android:name=\"com.oculus.permission.HAND_TRACKING\"/>");
+
+					// need for passthrough on MetaQuest
+					Text.AppendLine("\t<uses-feature android:name=\"com.oculus.feature.PASSTHROUGH\" android:required=\"false\" />");
 				}
 
 				Text.AppendLine("\t<uses-feature android:glEsVersion=\"" + AndroidToolChain.GetGLESVersion(bBuildForES31) + "\" android:required=\"true\" />");
