PR #12978: [OpenXRHMD] Add PostBeginFrame_RHIThread

diff --git a/Engine/Plugins/Runtime/OpenXR/Source/OpenXRHMD/Private/OpenXRHMD.cpp b/Engine/Plugins/Runtime/OpenXR/Source/OpenXRHMD/Private/OpenXRHMD.cpp
--- a/Engine/Plugins/Runtime/OpenXR/Source/OpenXRHMD/Private/OpenXRHMD.cpp
+++ b/Engine/Plugins/Runtime/OpenXR/Source/OpenXRHMD/Private/OpenXRHMD.cpp
@@ -3918,6 +3918,11 @@ void FOpenXRHMD::OnBeginRendering_RHIThread(IRHICommandContext& RHICmdContext, c
 		UE_LOG(LogHMD, VeryVerbose, TEXT("%s WF_%i Rendering frame predicted to be displayed at %lld"), 
 			   HMDThreadString(), PipelinedFrameStateRHI.WaitCount,
 			   PipelinedFrameStateRHI.FrameState.predictedDisplayTime);
+
+		for (IOpenXRExtensionPlugin* Module : ExtensionPlugins)
+		{
+			Module->PostBeginFrame_RHIThread(PipelinedFrameStateRHI.FrameState.predictedDisplayTime);
+		}
 	}
 	else
 	{
diff --git a/Engine/Plugins/Runtime/OpenXR/Source/OpenXRHMD/Public/IOpenXRExtensionPlugin.h b/Engine/Plugins/Runtime/OpenXR/Source/OpenXRHMD/Public/IOpenXRExtensionPlugin.h
--- a/Engine/Plugins/Runtime/OpenXR/Source/OpenXRHMD/Public/IOpenXRExtensionPlugin.h
+++ b/Engine/Plugins/Runtime/OpenXR/Source/OpenXRHMD/Public/IOpenXRExtensionPlugin.h
@@ -400,6 +400,10 @@ public:
 		PRAGMA_ENABLE_DEPRECATION_WARNINGS
 	}
 
+	virtual void PostBeginFrame_RHIThread(XrTime PredictedDisplayTime)
+	{
+	}
+
 	// OpenXRHMD::OnBeginRendering_RHIThread
 	virtual const void* OnBeginFrame_RHIThread(XrSession InSession, XrTime DisplayTime, const void* InNext)	
 	{
