visionOS uninitialized bRuntimeRequiresRHIContext
-This member was uninitialized.  If it is randomly true we would pass this extension into other openxr runtimes, which they should ignore.  Noticed while debugging the attached jira.

diff --git a/Engine/Plugins/Runtime/OpenXR/Source/OpenXRHMD/Private/OpenXRHMD.cpp b/Engine/Plugins/Runtime/OpenXR/Source/OpenXRHMD/Private/OpenXRHMD.cpp
--- a/Engine/Plugins/Runtime/OpenXR/Source/OpenXRHMD/Private/OpenXRHMD.cpp
+++ b/Engine/Plugins/Runtime/OpenXR/Source/OpenXRHMD/Private/OpenXRHMD.cpp
@@ -1637,6 +1637,7 @@ FOpenXRHMD::FOpenXRHMD(const FAutoRegister& AutoRegister, XrInstance InInstance,
 	, bIsMobileMultiViewEnabled(false)
 	, bSupportsHandTracking(false)
 	, bIsStandaloneStereoOnlyDevice(false)
+	, bRuntimeRequiresRHIContext(false)
 	, bIsTrackingOnlySession(false)
 	, CurrentSessionState(XR_SESSION_STATE_UNKNOWN)
 	, EnabledExtensions(std::move(InEnabledExtensions))
