Update OpenXRHMD to run all OpenXR calls which can access the VkQueue (xrBeginFrame, xrEndFrame, xrAcquireSwapchainImage, xrReleaseSwapchainImage) via RHIRunOnQueue when using Vulkan (if they weren't already called from Present(), which flushes earlier commands), preventing a crash when using the Vulkan submission thread.
This is a minimal hotfix-compliant fix and should not be merged to main.

diff --git a/Engine/Plugins/Runtime/OpenXR/Source/OpenXRHMD/Private/OpenXRHMD.cpp b/Engine/Plugins/Runtime/OpenXR/Source/OpenXRHMD/Private/OpenXRHMD.cpp
--- a/Engine/Plugins/Runtime/OpenXR/Source/OpenXRHMD/Private/OpenXRHMD.cpp
+++ b/Engine/Plugins/Runtime/OpenXR/Source/OpenXRHMD/Private/OpenXRHMD.cpp
@@ -191,6 +191,25 @@ namespace {
 	#endif
 		return nullptr;
 	}
+
+	/** Hotfix 5.6.1: Helper function to run OpenXR calls which may access VkQueue on the Vulkan submission thread */
+	void RunOnGraphicsQueue(TFunction<void()>&& CodeToRun)
+	{
+#ifdef XR_USE_GRAPHICS_API_VULKAN
+		if (RHIGetInterfaceType() == ERHIInterfaceType::Vulkan)
+		{
+			GetIVulkanDynamicRHI()->RHIRunOnQueue(EVulkanRHIRunOnQueueType::Graphics, [CodeToRun](VkQueue Queue)
+			{
+				FTaskTagScope Scope(ETaskTag::EParallelRhiThread);
+				CodeToRun();
+			}, false);
+		}
+		else
+#endif
+		{
+			CodeToRun();
+		}
+	}
 }
 
 //---------------------------------------------------
@@ -2724,7 +2743,10 @@ bool FOpenXRHMD::AllocateRenderTargetTextures(uint32 SizeX, uint32 SizeY, uint8
 			FRHICommandListImmediate& RHICmdList = FRHICommandListImmediate::Get();
 			RHICmdList.EnqueueLambda([Swapchain](FRHICommandListImmediate&)
 			{
-				Swapchain->IncrementSwapChainIndex_RHIThread();
+				RunOnGraphicsQueue([Swapchain]()
+				{
+					Swapchain->IncrementSwapChainIndex_RHIThread();
+				});
 			});
 			RHICmdList.ImmediateFlush(EImmediateFlushType::FlushRHIThread);
 		}
@@ -2759,7 +2781,10 @@ bool FOpenXRHMD::AllocateRenderTargetTextures(uint32 SizeX, uint32 SizeY, uint8
 				FRHICommandListImmediate& RHICmdList = FRHICommandListImmediate::Get();
 				RHICmdList.EnqueueLambda([EmulationSwapchain](FRHICommandListImmediate&)
 				{
-					EmulationSwapchain->IncrementSwapChainIndex_RHIThread();
+					RunOnGraphicsQueue([EmulationSwapchain]
+					{
+						EmulationSwapchain->IncrementSwapChainIndex_RHIThread();
+					});
 				});
 				RHICmdList.ImmediateFlush(EImmediateFlushType::FlushRHIThread);
 			}
@@ -2815,7 +2840,10 @@ void FOpenXRHMD::AllocateDepthTextureInternal(uint32 SizeX, uint32 SizeY, uint32
 			FRHICommandListImmediate& RHICmdList = FRHICommandListImmediate::Get();
 			RHICmdList.EnqueueLambda([DepthSwapchain](FRHICommandListImmediate&)
 			{
-				DepthSwapchain->IncrementSwapChainIndex_RHIThread();
+				RunOnGraphicsQueue([DepthSwapchain]()
+				{
+					DepthSwapchain->IncrementSwapChainIndex_RHIThread();
+				});
 			});
 			RHICmdList.ImmediateFlush(EImmediateFlushType::FlushRHIThread);
 		}
@@ -3388,7 +3416,11 @@ void FOpenXRHMD::OnBeginRendering_RenderThread(FRDGBuilder& GraphBuilder, FScene
 		// For now, leaving swapchain acquisition outside of an RDG pass to make sure the swapchain is acquired early enough.
 		GraphBuilder.RHICmdList.EnqueueLambda([this, FrameState = PipelinedFrameStateRendering, ColorSwapchain, DepthSwapchain, EmulationSwapchain](FRHICommandListImmediate& InRHICmdList)
 		{
-			OnBeginRendering_RHIThread(InRHICmdList.GetContext(), FrameState, ColorSwapchain, DepthSwapchain, EmulationSwapchain);
+			IRHICommandContext* Context = &InRHICmdList.GetContext();
+			RunOnGraphicsQueue([this, FrameState = PipelinedFrameStateRendering, ColorSwapchain, DepthSwapchain, EmulationSwapchain, Context]()
+			{
+				OnBeginRendering_RHIThread(*Context, FrameState, ColorSwapchain, DepthSwapchain, EmulationSwapchain);
+			});
 		});
 	}
 }
@@ -3836,7 +3868,7 @@ void FOpenXRHMD::RequestExitApp()
 
 void FOpenXRHMD::OnBeginRendering_RHIThread(IRHICommandContext& RHICmdContext, const FPipelinedFrameState& InFrameState, FXRSwapChainPtr ColorSwapchain, FXRSwapChainPtr DepthSwapchain, FXRSwapChainPtr EmulationSwapchain)
 {
-	ensure(IsInRenderingThread() || IsInRHIThread());
+	ensure(IsInRenderingThread() || IsInRHIThread() || IsInParallelRHIThread());
 
 	// TODO: Add a hook to resolve discarded frames before we start a new frame.
 	UE_CLOG(bIsRendering, LogHMD, Verbose, TEXT("Discarded previous frame and started rendering a new frame."));
@@ -3947,7 +3979,7 @@ void FOpenXRHMD::OnBeginRendering_RHIThread(IRHICommandContext& RHICmdContext, c
 
 void FOpenXRHMD::OnFinishRendering_RHIThread(IRHICommandContext& RHICmdContext)
 {
-	ensure(IsInRenderingThread() || IsInRHIThread());
+	ensure(IsInRenderingThread() || IsInRHIThread() || IsInParallelRHIThread());
 
 	SCOPED_NAMED_EVENT(EndFrame, FColor::Red);
 
@@ -4188,8 +4220,11 @@ void FOpenXRHMD::CopySwapchainTexture_RenderThread(FRDGBuilder& GraphBuilder, FR
 	{
 		RHICmdList.EnqueueLambda([DstSwapChain](FRHICommandListImmediate& InRHICmdList)
 		{
-			DstSwapChain->IncrementSwapChainIndex_RHIThread();
-			DstSwapChain->WaitCurrentImage_RHIThread(OPENXR_SWAPCHAIN_WAIT_TIMEOUT);
+			RunOnGraphicsQueue([DstSwapChain]()
+			{
+				DstSwapChain->IncrementSwapChainIndex_RHIThread();
+				DstSwapChain->WaitCurrentImage_RHIThread(OPENXR_SWAPCHAIN_WAIT_TIMEOUT);
+			});
 		});
 	});
 
@@ -4211,7 +4246,10 @@ void FOpenXRHMD::CopySwapchainTexture_RenderThread(FRDGBuilder& GraphBuilder, FR
 	{
 		RHICmdList.EnqueueLambda([DstSwapChain, bCapturableRequiresRHIContext](FRHICommandListImmediate& InRHICmdList)
 		{
-			DstSwapChain->ReleaseCurrentImage_RHIThread(bCapturableRequiresRHIContext ? &InRHICmdList.GetContext() : nullptr);
+			RunOnGraphicsQueue([DstSwapChain, bCapturableRequiresRHIContext, &InRHICmdList]()
+			{
+				DstSwapChain->ReleaseCurrentImage_RHIThread(bCapturableRequiresRHIContext ? &InRHICmdList.GetContext() : nullptr);
+			});
 		});
 	});
 }
diff --git a/Engine/Plugins/Runtime/OpenXR/Source/OpenXRHMD/Private/OpenXRHMD_Swapchain.cpp b/Engine/Plugins/Runtime/OpenXR/Source/OpenXRHMD/Private/OpenXRHMD_Swapchain.cpp
--- a/Engine/Plugins/Runtime/OpenXR/Source/OpenXRHMD/Private/OpenXRHMD_Swapchain.cpp
+++ b/Engine/Plugins/Runtime/OpenXR/Source/OpenXRHMD/Private/OpenXRHMD_Swapchain.cpp
@@ -52,7 +52,7 @@ void FOpenXRSwapchain::IncrementSwapChainIndex_RHIThread()
 
 void FOpenXRSwapchain::WaitCurrentImage_RHIThread(int64 Timeout)
 {
-	check(IsInRenderingThread() || IsInRHIThread());
+	check(IsInRenderingThread() || IsInRHIThread() || IsInParallelRHIThread());
 
 	bool WasAcquired = true;
 	ImageAcquired.compare_exchange_strong(WasAcquired, false);
@@ -99,7 +99,7 @@ void FOpenXRSwapchain::WaitCurrentImage_RHIThread(int64 Timeout)
 
 void FOpenXRSwapchain::ReleaseCurrentImage_RHIThread(IRHICommandContext* RHICmdContext)
 {
-	check(IsInRenderingThread() || IsInRHIThread());
+	check(IsInRenderingThread() || IsInRHIThread() || IsInParallelRHIThread());
 
 	bool WasReady = true;
 	ImageReady.compare_exchange_strong(WasReady, false);
