Replace some usages of FORCEINLINE with inline in AR modules.

For context about this change, see this doc: https://docs.google.com/document/d/1LjpT7NpzTecfXdFNyv5KM6sCu67827BRljNi1iaWVSE/edit?tab=t.0#heading=h.fvtz8gjukz7x

Includes modules: AppleARKit, AppleARKitFaceSupport, OpenXRHMD, AugmentedReality, HeadMountedDisplay

Reviewers were automatically determined. Please comment if anyone else should be added.

#rnx

diff --git a/Engine/Plugins/Runtime/AR/AppleAR/AppleARKit/Source/AppleARKit/Public/AppleARKitConversion.h b/Engine/Plugins/Runtime/AR/AppleAR/AppleARKit/Source/AppleARKit/Public/AppleARKitConversion.h
--- a/Engine/Plugins/Runtime/AR/AppleAR/AppleARKit/Source/AppleARKit/Public/AppleARKitConversion.h
+++ b/Engine/Plugins/Runtime/AR/AppleAR/AppleARKit/Source/AppleARKit/Public/AppleARKitConversion.h
@@ -56,12 +56,12 @@ struct FARWorldSaveHeader
 
 struct APPLEARKIT_API FAppleARKitConversion
 {
-	static FORCEINLINE float ToUEScale()
+	static inline float ToUEScale()
 	{
 		return 100.f;
 	}
 
-	static FORCEINLINE float ToARKitScale()
+	static inline float ToARKitScale()
 	{
 		return 0.01f;
 	}
@@ -73,7 +73,7 @@ struct APPLEARKIT_API FAppleARKitConversion
 	 *
 	 * Scale is removed because some ARKit transforms are less normalized than unreal requires.
 	 */
-	static FORCEINLINE FTransform ToFTransform(const matrix_float4x4& RawYUpMatrix, const FRotator& AdjustBy = FRotator::ZeroRotator)
+	static inline FTransform ToFTransform(const matrix_float4x4& RawYUpMatrix, const FRotator& AdjustBy = FRotator::ZeroRotator)
 	{
 		// Conversion here is as per SteamVRHMD::ToFMatrix
 		FMatrix RawYUpFMatrix(
@@ -106,7 +106,7 @@ struct APPLEARKIT_API FAppleARKitConversion
      *
      * Ignores scale.
      */
-    static FORCEINLINE matrix_float4x4 ToARKitMatrix( const FTransform& InTransform, float WorldToMetersScale = 100.0f )
+    static inline matrix_float4x4 ToARKitMatrix( const FTransform& InTransform, float WorldToMetersScale = 100.0f )
     {
 		if(!ensure(WorldToMetersScale != 0.f))
 		{
@@ -138,7 +138,7 @@ struct APPLEARKIT_API FAppleARKitConversion
 	 * Convert's an ARKit 'Y up' 'right handed' coordinate system vector to Unreal's 'Z up' 
 	 * 'left handed' coordinate system.
 	 */
-	static FORCEINLINE FVector ToFVector(const vector_float3& RawYUpVector)
+	static inline FVector ToFVector(const vector_float3& RawYUpVector)
 	{
 		return FVector( -RawYUpVector.z, RawYUpVector.x, RawYUpVector.y ) * ToUEScale();
 	}
@@ -147,7 +147,7 @@ struct APPLEARKIT_API FAppleARKitConversion
      * Convert's an Unreal's 'Z up' to ARKit's 'Y up' vector
      * 'left handed' coordinate system.
      */
-    static FORCEINLINE vector_float3 ToARKitVector( const FVector& InFVector, float WorldToMetersScale = 100.0f )
+    static inline vector_float3 ToARKitVector( const FVector& InFVector, float WorldToMetersScale = 100.0f )
     {
 		if(!ensure(WorldToMetersScale != 0.f))
 		{
@@ -162,14 +162,14 @@ struct APPLEARKIT_API FAppleARKitConversion
         return RetVal / WorldToMetersScale;
     }
 
-	static FORCEINLINE FGuid ToFGuid( uuid_t UUID )
+	static inline FGuid ToFGuid( uuid_t UUID )
 	{
 		const uint32* Components = reinterpret_cast<const uint32*>(UUID);
 		FGuid AsGUID(Components[0], Components[1], Components[2], Components[3]);
 		return AsGUID;
 	}
 
-	static FORCEINLINE FGuid ToFGuid( NSUUID* Identifier )
+	static inline FGuid ToFGuid( NSUUID* Identifier )
 	{
 		// Get bytes
 		uuid_t UUID;
@@ -184,7 +184,7 @@ struct APPLEARKIT_API FAppleARKitConversion
 	/**
 	 * Coverts plane orientation
 	 */
-	static FORCEINLINE EARPlaneOrientation ToEARPlaneOrientation(ARPlaneAnchorAlignment Alignment)
+	static inline EARPlaneOrientation ToEARPlaneOrientation(ARPlaneAnchorAlignment Alignment)
 	{
 		EARPlaneOrientation RetVal = EARPlaneOrientation::Horizontal;
 		if (Alignment == ARPlaneAnchorAlignmentVertical)
@@ -195,7 +195,7 @@ struct APPLEARKIT_API FAppleARKitConversion
 	}
 
 #if SUPPORTS_ARKIT_2_0
-	static FORCEINLINE EARObjectClassification ToEARObjectClassification(ARPlaneClassification Classification)
+	static inline EARObjectClassification ToEARObjectClassification(ARPlaneClassification Classification)
 	{
 		switch(Classification)
 		{
@@ -229,7 +229,7 @@ struct APPLEARKIT_API FAppleARKitConversion
 #endif
 	
 #if SUPPORTS_ARKIT_3_5
-	static FORCEINLINE EARObjectClassification ToEARObjectClassification(ARMeshClassification Classification)
+	static inline EARObjectClassification ToEARObjectClassification(ARMeshClassification Classification)
 	{
 		switch (Classification)
 		{
diff --git a/Engine/Plugins/Runtime/AR/AppleAR/AppleARKitFaceSupport/Source/AppleARKitFaceSupport/Public/AppleARKitFaceMeshComponent.h b/Engine/Plugins/Runtime/AR/AppleAR/AppleARKitFaceSupport/Source/AppleARKitFaceSupport/Public/AppleARKitFaceMeshComponent.h
--- a/Engine/Plugins/Runtime/AR/AppleAR/AppleARKitFaceSupport/Source/AppleARKitFaceSupport/Public/AppleARKitFaceMeshComponent.h
+++ b/Engine/Plugins/Runtime/AR/AppleAR/AppleARKitFaceSupport/Source/AppleARKitFaceSupport/Public/AppleARKitFaceMeshComponent.h
@@ -31,22 +31,22 @@ struct FNetQuantizeFaceCurve
 {
 	GENERATED_USTRUCT_BODY()
 
-	FORCEINLINE FNetQuantizeFaceCurve()
+	inline FNetQuantizeFaceCurve()
 	{
 	}
 
-	FORCEINLINE FNetQuantizeFaceCurve(EARFaceBlendShape InBlendShape, float InAmount)
+	inline FNetQuantizeFaceCurve(EARFaceBlendShape InBlendShape, float InAmount)
 		: BlendShape((uint8)InBlendShape)
 		, Amount(ConvertAmountToInt(InAmount))
 	{
 	}
 
-	FORCEINLINE float GetAmountAsFloat() const
+	inline float GetAmountAsFloat() const
 	{
 		return ConvertAmountToFloat(Amount);
 	}
 
-	FORCEINLINE EARFaceBlendShape GetBlendShape() const
+	inline EARFaceBlendShape GetBlendShape() const
 	{
 		check(BlendShape < (uint8)EARFaceBlendShape::MAX);
 		return (EARFaceBlendShape)BlendShape;
@@ -79,12 +79,12 @@ private:
 	static constexpr float Scale = 127.f;
 	static constexpr float InvScale = 1.f / Scale;
 	
-	static FORCEINLINE float ConvertAmountToFloat(int8 InAmount)
+	static inline float ConvertAmountToFloat(int8 InAmount)
 	{
 		return ((float)InAmount) * InvScale;
 	}
 
-	static FORCEINLINE int8 ConvertAmountToInt(float InAmount)
+	static inline int8 ConvertAmountToInt(float InAmount)
 	{
 		return (int8)FMath::TruncToInt(InAmount * Scale);
 	}
diff --git a/Engine/Plugins/Runtime/OpenXR/Source/OpenXRHMD/Public/OpenXRCore.h b/Engine/Plugins/Runtime/OpenXR/Source/OpenXRHMD/Public/OpenXRCore.h
--- a/Engine/Plugins/Runtime/OpenXR/Source/OpenXRHMD/Public/OpenXRCore.h
+++ b/Engine/Plugins/Runtime/OpenXR/Source/OpenXRHMD/Public/OpenXRCore.h
@@ -46,27 +46,27 @@ constexpr const TCHAR* OpenXRReferenceSpaceTypeToString(XrReferenceSpaceType e)
 #define XR_ENSURE(x) XR_SUCCEEDED(x)
 #endif
 
-FORCEINLINE FQuat ToFQuat(XrQuaternionf Quat)
+inline FQuat ToFQuat(XrQuaternionf Quat)
 {
 	return FQuat(-Quat.z, Quat.x, Quat.y, -Quat.w);
 }
 
-FORCEINLINE XrQuaternionf ToXrQuat(FQuat Quat)
+inline XrQuaternionf ToXrQuat(FQuat Quat)
 {
 	return XrQuaternionf{ (float)Quat.Y, (float)Quat.Z, -(float)Quat.X, -(float)Quat.W };
 }
 
-FORCEINLINE FVector ToFVector(XrVector3f Vector, float Scale = 1.0f)
+inline FVector ToFVector(XrVector3f Vector, float Scale = 1.0f)
 {
 	return FVector(-Vector.z * Scale, Vector.x * Scale, Vector.y * Scale);
 }
 
-FORCEINLINE FVector3f ToFVector3f(XrVector3f Vector, float Scale = 1.0f)
+inline FVector3f ToFVector3f(XrVector3f Vector, float Scale = 1.0f)
 {
 	return FVector3f(-Vector.z * Scale, Vector.x * Scale, Vector.y * Scale);
 }
 
-FORCEINLINE XrVector3f ToXrVector(FVector Vector, float Scale = 1.0f)
+inline XrVector3f ToXrVector(FVector Vector, float Scale = 1.0f)
 {
 	if (Vector.IsZero())
 		return XrVector3f{ 0.0f, 0.0f, 0.0f };
@@ -74,54 +74,54 @@ FORCEINLINE XrVector3f ToXrVector(FVector Vector, float Scale = 1.0f)
 	return XrVector3f{ (float)Vector.Y / Scale, (float)Vector.Z / Scale, (float)-Vector.X / Scale };
 }
 
-FORCEINLINE FTransform ToFTransform(XrPosef Transform, float Scale = 1.0f)
+inline FTransform ToFTransform(XrPosef Transform, float Scale = 1.0f)
 {
 	return FTransform(ToFQuat(Transform.orientation), ToFVector(Transform.position, Scale));
 }
 
-FORCEINLINE XrPosef ToXrPose(FTransform Transform, float Scale = 1.0f)
+inline XrPosef ToXrPose(FTransform Transform, float Scale = 1.0f)
 {
 	return XrPosef{ ToXrQuat(Transform.GetRotation()), ToXrVector(Transform.GetTranslation(), Scale) };
 }
 
-FORCEINLINE FTimespan ToFTimespan(XrTime Time)
+inline FTimespan ToFTimespan(XrTime Time)
 {
 	// XrTime is a nanosecond counter, FTimespan is a 100-nanosecond counter. 
 	// We are losing some precision here.
 	return FTimespan((Time + 50) / 100); 
 }
 
-FORCEINLINE XrTime ToXrTime(FTimespan Time)
+inline XrTime ToXrTime(FTimespan Time)
 {
 	return Time.GetTicks() * 100;
 }
 
-FORCEINLINE FIntRect ToFIntRect(XrRect2Di Rect)
+inline FIntRect ToFIntRect(XrRect2Di Rect)
 {
 	return FIntRect(Rect.offset.x, Rect.offset.y, Rect.offset.x + Rect.extent.width, Rect.offset.y + Rect.extent.height);
 }
 
-FORCEINLINE XrRect2Di ToXrRect(FIntRect Rect)
+inline XrRect2Di ToXrRect(FIntRect Rect)
 {
 	return XrRect2Di{ { Rect.Min.X, Rect.Min.Y }, { Rect.Width(), Rect.Height() } };
 }
 
-FORCEINLINE FVector2D ToFVector2D(XrVector2f Vector, float Scale = 1.0f)
+inline FVector2D ToFVector2D(XrVector2f Vector, float Scale = 1.0f)
 {
 	return FVector2D(Vector.x * Scale, Vector.y * Scale);
 }
 
-FORCEINLINE XrVector2f ToXrVector2f(FVector2D Vector, float Scale = 1.0f)
+inline XrVector2f ToXrVector2f(FVector2D Vector, float Scale = 1.0f)
 {
 	return XrVector2f{ (float)Vector.X / Scale, (float)Vector.Y / Scale };
 }
 
-FORCEINLINE FVector2D ToFVector2D(XrExtent2Df Extent, float Scale = 1.0f)
+inline FVector2D ToFVector2D(XrExtent2Df Extent, float Scale = 1.0f)
 {
 	return FVector2D(Extent.width * Scale, Extent.height * Scale);
 }
 
-FORCEINLINE XrExtent2Df ToXrExtent2D(FVector2D Vector, float Scale = 1.0f)
+inline XrExtent2Df ToXrExtent2D(FVector2D Vector, float Scale = 1.0f)
 {
 	if (Vector.IsZero())
 		return XrExtent2Df{ 0.0f, 0.0f };
@@ -129,7 +129,7 @@ FORCEINLINE XrExtent2Df ToXrExtent2D(FVector2D Vector, float Scale = 1.0f)
 	return XrExtent2Df{ (float)Vector.X / Scale, (float)Vector.Y / Scale };
 }
 
-FORCEINLINE uint32 ToXrPriority(int32 Priority)
+inline uint32 ToXrPriority(int32 Priority)
 {
 	// Ensure negative priority numbers map to the lower half of the 32-bit range.
 	// We do this by casting to an unsigned int and then flipping the signed bit.
@@ -299,7 +299,7 @@ bool InitOpenXRCore(XrInstance Instance);
 
 void OPENXRHMD_API EnumerateOpenXRApiLayers(TArray<XrApiLayerProperties>& OutProperties);
 
-FORCEINLINE void FilterActionName(const char* InActionName, char* OutActionName)
+inline void FilterActionName(const char* InActionName, char* OutActionName)
 {
 	static_assert(XR_MAX_ACTION_NAME_SIZE == XR_MAX_ACTION_SET_NAME_SIZE);
 
diff --git a/Engine/Source/Runtime/AugmentedReality/Public/ARTraceResult.h b/Engine/Source/Runtime/AugmentedReality/Public/ARTraceResult.h
--- a/Engine/Source/Runtime/AugmentedReality/Public/ARTraceResult.h
+++ b/Engine/Source/Runtime/AugmentedReality/Public/ARTraceResult.h
@@ -37,7 +37,7 @@ struct FARTraceResult
 
 	struct FARTraceResultComparer
 	{
-		FORCEINLINE_STATS bool operator()(const FARTraceResult& A, const FARTraceResult& B) const
+		inline bool operator()(const FARTraceResult& A, const FARTraceResult& B) const
 		{
 			return A.GetDistanceFromCamera() < B.GetDistanceFromCamera();
 		}
diff --git a/Engine/Source/Runtime/HeadMountedDisplay/Public/IMotionController.h b/Engine/Source/Runtime/HeadMountedDisplay/Public/IMotionController.h
--- a/Engine/Source/Runtime/HeadMountedDisplay/Public/IMotionController.h
+++ b/Engine/Source/Runtime/HeadMountedDisplay/Public/IMotionController.h
@@ -116,7 +116,7 @@ public:
 	*/
 	virtual bool GetControllerOrientationAndPositionForTime(const int32 ControllerIndex, const FName MotionSource, FTimespan Time, bool& OutTimeWasUsed, FRotator& OutOrientation, FVector& OutPosition, bool& OutbProvidedLinearVelocity, FVector& OutLinearVelocity, bool& OutbProvidedAngularVelocity, FVector& OutAngularVelocityAsAxisAndLength, bool& OutbProvidedLinearAcceleration, FVector& OutLinearAcceleration, float WorldToMetersScale) const = 0;
 
-	static FORCEINLINE FRotator AngularVelocityAsAxisAndLengthToRotator(const FVector AngularVelocityAsAxisAndLength)
+	static inline FRotator AngularVelocityAsAxisAndLengthToRotator(const FVector AngularVelocityAsAxisAndLength)
 	{
 		FVector Direction;
 		double Angle;
diff --git a/Engine/Source/Runtime/HeadMountedDisplay/Public/MotionDelayBuffer.inl b/Engine/Source/Runtime/HeadMountedDisplay/Public/MotionDelayBuffer.inl
--- a/Engine/Source/Runtime/HeadMountedDisplay/Public/MotionDelayBuffer.inl
+++ b/Engine/Source/Runtime/HeadMountedDisplay/Public/MotionDelayBuffer.inl
@@ -41,14 +41,14 @@ void TCircularHistoryBuffer<ElementType>::Resize(uint32 NewCapacity)
 }
 
 template<typename ElementType>
-FORCEINLINE ElementType& TCircularHistoryBuffer<ElementType>::operator[](uint32 Index)
+inline ElementType& TCircularHistoryBuffer<ElementType>::operator[](uint32 Index)
 {
 	RangeCheck(Index);
 	return Elements[AsInternalIndex(Index)];
 }
 
 template<typename ElementType>
-FORCEINLINE const ElementType& TCircularHistoryBuffer<ElementType>::operator[](uint32 Index) const
+inline const ElementType& TCircularHistoryBuffer<ElementType>::operator[](uint32 Index) const
 {
 	RangeCheck(Index);
 	return Elements[AsInternalIndex(Index)];
@@ -93,13 +93,13 @@ void TCircularHistoryBuffer<ElementType>::InsertAt(uint32 Index, const ElementTy
 }
 
 template<typename ElementType>
-FORCEINLINE uint32 TCircularHistoryBuffer<ElementType>::Capacity() const
+inline uint32 TCircularHistoryBuffer<ElementType>::Capacity() const
 {
 	return Elements.Num();
 }
 
 template<typename ElementType>
-FORCEINLINE int32 TCircularHistoryBuffer<ElementType>::Num() const
+inline int32 TCircularHistoryBuffer<ElementType>::Num() const
 {
 	if (IsFull())
 	{
@@ -112,26 +112,26 @@ FORCEINLINE int32 TCircularHistoryBuffer<ElementType>::Num() const
 }
 
 template<typename ElementType>
-FORCEINLINE void TCircularHistoryBuffer<ElementType>::Empty()
+inline void TCircularHistoryBuffer<ElementType>::Empty()
 {
 	Head = 0;
 	bIsFull = false;
 }
 
 template<typename ElementType>
-FORCEINLINE bool TCircularHistoryBuffer<ElementType>::IsEmpty() const
+inline bool TCircularHistoryBuffer<ElementType>::IsEmpty() const
 {
 	return (Head == 0) && !bIsFull;
 }
 
 template<typename ElementType>
-FORCEINLINE bool TCircularHistoryBuffer<ElementType>::IsFull() const
+inline bool TCircularHistoryBuffer<ElementType>::IsFull() const
 {
 	return bIsFull;
 }
 
 template<typename ElementType>
-FORCEINLINE void TCircularHistoryBuffer<ElementType>::RangeCheck(const uint32 Index, bool bCheckIfUnderfilled) const
+inline void TCircularHistoryBuffer<ElementType>::RangeCheck(const uint32 Index, bool bCheckIfUnderfilled) const
 {
 	checkSlow(!IsEmpty());
 	checkSlow(Index < (uint32)Elements.Num());
@@ -139,7 +139,7 @@ FORCEINLINE void TCircularHistoryBuffer<ElementType>::RangeCheck(const uint32 In
 }
 
 template<typename ElementType>
-FORCEINLINE uint32 TCircularHistoryBuffer<ElementType>::AsInternalIndex(uint32 Index) const
+inline uint32 TCircularHistoryBuffer<ElementType>::AsInternalIndex(uint32 Index) const
 {
 	//        head(3)                                            head(3)
 	//           |                                                  |
@@ -148,7 +148,7 @@ FORCEINLINE uint32 TCircularHistoryBuffer<ElementType>::AsInternalIndex(uint32 I
 }
 
 template<typename ElementType>
-FORCEINLINE uint32 TCircularHistoryBuffer<ElementType>::GetNextIndex(uint32 CurrentIndex) const
+inline uint32 TCircularHistoryBuffer<ElementType>::GetNextIndex(uint32 CurrentIndex) const
 {
 	return ((CurrentIndex + 1) % Elements.Num());
 }
