Exposing full FOV information for OpenXR HMDs
-Adds an accessor for the fov structs.

diff --git a/Engine/Plugins/Runtime/OpenXR/Source/OpenXRHMD/Private/OpenXRHMD.cpp b/Engine/Plugins/Runtime/OpenXR/Source/OpenXRHMD/Private/OpenXRHMD.cpp
--- a/Engine/Plugins/Runtime/OpenXR/Source/OpenXRHMD/Private/OpenXRHMD.cpp
+++ b/Engine/Plugins/Runtime/OpenXR/Source/OpenXRHMD/Private/OpenXRHMD.cpp
@@ -811,6 +811,18 @@ void FOpenXRHMD::GetFieldOfView(float& OutHFOVInDegrees, float& OutVFOVInDegrees
 	OutVFOVInDegrees = FMath::RadiansToDegrees(UnifiedFov.angleUp - UnifiedFov.angleDown);
 }
 
+void FOpenXRHMD::GetFullFOVInformation(TArray<XrFovf>& FovInfos) const
+{
+	FovInfos.Empty();
+
+	FPipelinedFrameStateAccessorReadOnly LockedPipelineState = GetPipelinedFrameStateForThread();
+	const FPipelinedFrameState& FrameState = LockedPipelineState.GetFrameState();
+	for (const XrView& View : FrameState.Views)
+	{
+		FovInfos.Add(View.fov);
+	}
+}
+
 bool FOpenXRHMD::EnumerateTrackedDevices(TArray<int32>& OutDevices, EXRTrackedDeviceType Type)
 {
 	if (Type == EXRTrackedDeviceType::Any || Type == EXRTrackedDeviceType::HeadMountedDisplay)
diff --git a/Engine/Plugins/Runtime/OpenXR/Source/OpenXRHMD/Private/OpenXRHMD.h b/Engine/Plugins/Runtime/OpenXR/Source/OpenXRHMD/Private/OpenXRHMD.h
--- a/Engine/Plugins/Runtime/OpenXR/Source/OpenXRHMD/Private/OpenXRHMD.h
+++ b/Engine/Plugins/Runtime/OpenXR/Source/OpenXRHMD/Private/OpenXRHMD.h
@@ -484,8 +484,15 @@ public:
 	IOpenXRExtensionPluginDelegates& GetIOpenXRExtensionPluginDelegates() override { return *this; }
 	TArray<IOpenXRExtensionPlugin*>& GetExtensionPlugins() override { return ExtensionPlugins; }
 	
+
+	/* Note: this function will be added to IOpenXRHMD.h in 5.7*/
 	OPENXRHMD_API void SetEnvironmentBlendMode(XrEnvironmentBlendMode NewBlendMode);
 
+	/** Provides the full asymmetric FOV information for the view, which allows proper centering of screen-based techniques based on forward eye gaze 
+	* Note: this function will only be exposed here in 5.6.  In 5.7 it will be added to IOpenXRHMD.h.
+	*/
+	OPENXRHMD_API void GetFullFOVInformation(TArray<XrFovf>& FovInfos) const;
+
 	/** Returns shader platform the plugin is currently configured for, in the editor it can change due to preview platforms. */
 	EShaderPlatform GetConfiguredShaderPlatform() const { check(ConfiguredShaderPlatform != EShaderPlatform::SP_NumPlatforms); return ConfiguredShaderPlatform; }
 	FOpenXRSwapchain* GetColorSwapchain_RenderThread();
