Bug fix for FOpenXRHMD::GetCurrentInteractionProfile and related methods.
-The right hand was malfunctioning in some configurations and always returned the left hand information.

diff --git a/Engine/Plugins/Runtime/OpenXR/Source/OpenXRHMD/Private/OpenXRHMD.cpp b/Engine/Plugins/Runtime/OpenXR/Source/OpenXRHMD/Private/OpenXRHMD.cpp
--- a/Engine/Plugins/Runtime/OpenXR/Source/OpenXRHMD/Private/OpenXRHMD.cpp
+++ b/Engine/Plugins/Runtime/OpenXR/Source/OpenXRHMD/Private/OpenXRHMD.cpp
@@ -271,19 +271,10 @@ void FOpenXRHMD::GetMotionControllerData(UObject* WorldContext, const EControlle
 	MotionControllerData.HandIndex = Hand;
 	MotionControllerData.bValid = false;
 
-	TArray<int32> Devices;
-	if (EnumerateTrackedDevices(Devices, EXRTrackedDeviceType::Controller) && Devices.IsValidIndex((int32)Hand))
+	FString InteractionProfile;
+	if (GetCurrentInteractionProfile(Hand, InteractionProfile))
 	{
-		FReadScopeLock SessionLock(SessionHandleMutex);
-		if (Session)
-		{
-			XrInteractionProfileState Profile = { XR_TYPE_INTERACTION_PROFILE_STATE };
-			if (XR_SUCCEEDED(xrGetCurrentInteractionProfile(Session, GetTrackedDevicePath(Devices[(int32)Hand]), &Profile)) &&
-				Profile.interactionProfile != XR_NULL_PATH)
-			{
-				MotionControllerData.DeviceName = FOpenXRPath(Profile.interactionProfile);
-			}
-		}
+		MotionControllerData.DeviceName = FName(InteractionProfile);
 	}
 
 	FName HandTrackerName("OpenXRHandTracking");
@@ -417,19 +408,10 @@ void FOpenXRHMD::GetMotionControllerState(UObject* WorldContext, const EXRSpaceT
 	MotionControllerState.XRSpaceType = XRSpaceType;
 	MotionControllerState.bValid = false;
 
-	TArray<int32> Devices;
-	if (EnumerateTrackedDevices(Devices, EXRTrackedDeviceType::Controller) && Devices.IsValidIndex((int32)Hand))
+	FString InteractionProfile;
+	if (GetCurrentInteractionProfile(Hand, InteractionProfile))
 	{
-		FReadScopeLock SessionLock(SessionHandleMutex);
-		if (Session)
-		{
-			XrInteractionProfileState Profile = { XR_TYPE_INTERACTION_PROFILE_STATE };
-			if (XR_SUCCEEDED(xrGetCurrentInteractionProfile(Session, GetTrackedDevicePath(Devices[(int32)Hand]), &Profile)) &&
-				Profile.interactionProfile != XR_NULL_PATH)
-			{
-				MotionControllerState.DeviceName = FOpenXRPath(Profile.interactionProfile);
-			}
-		}
+		MotionControllerState.DeviceName = FName(InteractionProfile);
 	}
 
 	if ((Hand == EControllerHand::Left) || (Hand == EControllerHand::Right))
@@ -531,26 +513,10 @@ void FOpenXRHMD::GetHandTrackingState(UObject* WorldContext, const EXRSpaceType
 
 bool FOpenXRHMD::GetCurrentInteractionProfile(const EControllerHand Hand, FString& InteractionProfile)
 {
-	int32 DeviceId = -1;
-	if (Hand == EControllerHand::HMD)
+	const XrPath Path = GetUserPathForControllerHand(Hand);
+	if (Path == XR_NULL_PATH)
 	{
-		DeviceId = IXRTrackingSystem::HMDDeviceId;
-	}
-	else
-	{
-		TArray<int32> Devices;
-		if (EnumerateTrackedDevices(Devices, EXRTrackedDeviceType::Controller))
-		{
-			if (Devices.IsValidIndex((int32)Hand))
-			{
-				DeviceId = Devices[(int32)Hand];
-			}
-		}
-	}
-
-	if (DeviceId == -1)
-	{
-		UE_LOG(LogHMD, Warning, TEXT("GetCurrentInteractionProfile failed because that EControllerHandValue %i does not map to a device!"), int(Hand));
+		UE_LOG(LogHMD, Warning, TEXT("GetCurrentInteractionProfile failed because that EControllerHandValue %i does not map to a user path!"), int(Hand));
 		return false;
 	}
 
@@ -558,7 +524,6 @@ bool FOpenXRHMD::GetCurrentInteractionProfile(const EControllerHand Hand, FStrin
 	if (Session)
 	{
 		XrInteractionProfileState Profile{ XR_TYPE_INTERACTION_PROFILE_STATE };
-		XrPath Path = GetTrackedDevicePath(DeviceId);
 		XrResult Result = xrGetCurrentInteractionProfile(Session, Path, &Profile);
 		if (XR_SUCCEEDED(Result))
 		{
@@ -4466,4 +4431,25 @@ void FOpenXRHMD::FTrackingSpace::DestroySpace()
 	Handle = XR_NULL_HANDLE;
 }
 
+XrPath FOpenXRHMD::GetUserPathForControllerHand(EControllerHand Hand)
+{
+	XrPath Path = XR_NULL_PATH;
+	switch  (Hand)
+	{
+	case EControllerHand::Left:
+		Path = FOpenXRPath("/user/hand/left");
+		break;
+	case EControllerHand::Right:
+		Path = FOpenXRPath("/user/hand/right");
+		break;
+	case EControllerHand::HMD:
+		Path = FOpenXRPath("/user/head");
+		break;
+	default:
+		break;
+	}
+	
+	return Path;
+}
+
 #undef LOCTEXT_NAMESPACE
diff --git a/Engine/Plugins/Runtime/OpenXR/Source/OpenXRHMD/Private/OpenXRHMD.h b/Engine/Plugins/Runtime/OpenXR/Source/OpenXRHMD/Private/OpenXRHMD.h
--- a/Engine/Plugins/Runtime/OpenXR/Source/OpenXRHMD/Private/OpenXRHMD.h
+++ b/Engine/Plugins/Runtime/OpenXR/Source/OpenXRHMD/Private/OpenXRHMD.h
@@ -511,6 +511,9 @@ private:
 	void ConfigureLayerSwapchains(const struct FLayerToUpdateSwapchain& Update, FOpenXRLayer& Layer);
 	void AddLayersToHeaders(TArray<XrCompositionLayerBaseHeader*>& Headers);
 
+	// Gets the user path corresponding to the input EControllerHand value
+	static XrPath GetUserPathForControllerHand(EControllerHand Hand);
+
 	bool					bStereoEnabled;
 	TAtomic<bool>			bIsRunning;
 	TAtomic<bool>			bIsReady;
