Enabling LinuxArm64 support for OpenXR
-LinuxArm64 is now a target for the OpenXR plugins.
-This does *not* provide an openxr loader for linuxarm64.  The custom loader feature could be used to expose one to openxr.

diff --git a/Engine/Plugins/Runtime/OpenXR/OpenXR.uplugin b/Engine/Plugins/Runtime/OpenXR/OpenXR.uplugin
--- a/Engine/Plugins/Runtime/OpenXR/OpenXR.uplugin
+++ b/Engine/Plugins/Runtime/OpenXR/OpenXR.uplugin
@@ -17,6 +17,7 @@
 	"SupportedTargetPlatforms": [
 		"Win64",
 		"Linux",
+		"LinuxArm64",
 		"Android"
 	],
 	"Modules": [
@@ -27,6 +28,7 @@
 			"PlatformAllowList": [
 				"Win64",
 				"Linux",
+				"LinuxArm64",
 				"Android",
 				"Mac"
 			]
@@ -38,6 +40,7 @@
 			"PlatformAllowList": [
 				"Win64",
 				"Linux",
+				"LinuxArm64",
 				"Android",
 				"Mac"
 			]
@@ -49,6 +52,7 @@
 			"PlatformAllowList": [
 				"Win64",
 				"Linux",
+				"LinuxArm64",
 				"Android",
 				"Mac"
 			]
diff --git a/Engine/Plugins/Runtime/OpenXR/Source/OpenXRHMD/Private/OpenXRHMDModule.cpp b/Engine/Plugins/Runtime/OpenXR/Source/OpenXRHMD/Private/OpenXRHMDModule.cpp
--- a/Engine/Plugins/Runtime/OpenXR/Source/OpenXRHMD/Private/OpenXRHMDModule.cpp
+++ b/Engine/Plugins/Runtime/OpenXR/Source/OpenXRHMD/Private/OpenXRHMDModule.cpp
@@ -433,9 +433,20 @@ PFN_xrGetInstanceProcAddr FOpenXRHMDModule::GetDefaultLoader()
 	LoaderHandle = FPlatformProcess::GetDllHandle(*LoaderName);
 	FPlatformProcess::PopDllDirectory(*BinariesPath);
 #elif PLATFORM_LINUX
+
+#if PLATFORM_LINUXARM64
+	// We do not currently have a binary for linuxarm64, so we will fail to find a default loader.
+	// It is possible to work around this by using the CustomLoader functionality to specify a loader.
+	// When we get one we will need to activate the following code, and add the new binary in OpenXR.Build.cs
+	//FString BinariesPath = FPaths::EngineDir() / FString(TEXT("Binaries/ThirdParty/OpenXR/linuxarm64/aarch64-unknown-linux-gnueabi"));
+	//FString LoaderName = "libopenxr_loader.so";
+	//LoaderHandle = FPlatformProcess::GetDllHandle(*(BinariesPath / LoaderName)); 
+#else
 	FString BinariesPath = FPaths::EngineDir() / FString(TEXT("Binaries/ThirdParty/OpenXR/linux/x86_64-unknown-linux-gnu"));
 	FString LoaderName = "libopenxr_loader.so";
 	LoaderHandle = FPlatformProcess::GetDllHandle(*(BinariesPath / LoaderName)); 
+#endif
+
 #elif PLATFORM_ANDROID
 	FString LoaderName = "libopenxr_loader.so";
 	LoaderHandle = FPlatformProcess::GetDllHandle(*LoaderName);
diff --git a/Engine/Plugins/Runtime/OpenXREyeTracker/OpenXREyeTracker.uplugin b/Engine/Plugins/Runtime/OpenXREyeTracker/OpenXREyeTracker.uplugin
--- a/Engine/Plugins/Runtime/OpenXREyeTracker/OpenXREyeTracker.uplugin
+++ b/Engine/Plugins/Runtime/OpenXREyeTracker/OpenXREyeTracker.uplugin
@@ -17,6 +17,7 @@
 	"SupportedTargetPlatforms": [
 		"Win64",
 		"Linux",
+		"LinuxArm64",
 		"Android"
 	],
 	"Modules": [
@@ -27,6 +28,7 @@
 			"PlatformAllowList": [
 				"Win64",
 				"Linux",
+				"LinuxArm64",
 				"Android",
 				"Mac"
 			]
diff --git a/Engine/Plugins/Runtime/OpenXRHandTracking/OpenXRHandTracking.uplugin b/Engine/Plugins/Runtime/OpenXRHandTracking/OpenXRHandTracking.uplugin
--- a/Engine/Plugins/Runtime/OpenXRHandTracking/OpenXRHandTracking.uplugin
+++ b/Engine/Plugins/Runtime/OpenXRHandTracking/OpenXRHandTracking.uplugin
@@ -17,6 +17,7 @@
 	"SupportedTargetPlatforms": [
 		"Win64",
 		"Linux",
+		"LinuxArm64",
 		"Android"
 	],
 	"Modules": [
@@ -27,6 +28,7 @@
 			"PlatformAllowList": [
 				"Win64",
 				"Linux",
+				"LinuxArm64",
 				"Android",
 				"Mac"
 			]
@@ -49,6 +51,7 @@
 			"PlatformAllowList": [
 				"Win64",
 				"Linux",
+				"LinuxArm64",
 				"Android",
 				"Mac"
 			]
