[Open XR] Remove references to the previously deprecated UPlayerMappableInputConfig in prep for its removal

[RN] Delete the formerly deprecated functions from 5.3 to remove references to the UPlayerMappableInputConfig type.

diff --git a/Engine/Plugins/Runtime/OpenXR/Source/OpenXRInput/Private/OpenXRInput.cpp b/Engine/Plugins/Runtime/OpenXR/Source/OpenXRInput/Private/OpenXRInput.cpp
--- a/Engine/Plugins/Runtime/OpenXR/Source/OpenXRInput/Private/OpenXRInput.cpp
+++ b/Engine/Plugins/Runtime/OpenXR/Source/OpenXRInput/Private/OpenXRInput.cpp
@@ -437,9 +437,7 @@ bool FOpenXRInputPlugin::FOpenXRInput::BuildActions(XrSession Session)
 			}
 		}
 	}
-
-
-	PRAGMA_DISABLE_DEPRECATION_WARNINGS
+	
 	// Attempt to load the default input config from the OpenXR input settings.
 	const UEnhancedInputDeveloperSettings* InputSettings = GetDefault<UEnhancedInputDeveloperSettings>();
 	if (InputSettings && InputSettings->bEnableDefaultMappingContexts)
@@ -464,11 +462,10 @@ bool FOpenXRInputPlugin::FOpenXRInput::BuildActions(XrSession Session)
 	}
 	else
 	{
-		UE_LOG(LogHMD, Warning, TEXT("No UPlayerMappableInputConfig provided in the OpenXR Input project settings, action bindings will not be visible to the OpenXR runtime. Action/Axis mappings from the Input configuration are deprecated for XR in favor of Enhanced Input."));
+		UE_LOG(LogHMD, Warning, TEXT("No mapping context provided in the OpenXR Input project settings, action bindings will not be visible to the OpenXR runtime. Action/Axis mappings from the Input configuration are deprecated for XR in favor of Enhanced Input."));
 
 		BuildLegacyActions(Profiles);
 	}
-	PRAGMA_ENABLE_DEPRECATION_WARNINGS
 
 	for (TPair<FString, FInteractionProfile>& Pair : Profiles)
 	{
@@ -1539,22 +1536,6 @@ float FOpenXRInputPlugin::FOpenXRInput::GetHapticAmplitudeScale() const
 	return 1.0f;
 }
 
-PRAGMA_DISABLE_DEPRECATION_WARNINGS
-bool FOpenXRInputPlugin::FOpenXRInput::SetPlayerMappableInputConfig(TObjectPtr<class UPlayerMappableInputConfig> InputConfig)
-{
-	if (bActionsAttached)
-	{
-		UE_LOG(LogHMD, Error, TEXT("Attempted to attach an input config while one is already attached for the current session."));
-
-		return false;
-	}
-
-	TSet<TObjectPtr<UInputMappingContext>> MappingContexts;
-	InputConfig->GetMappingContexts().GetKeys(MappingContexts);
-	return AttachInputMappingContexts(MappingContexts);
-}
-PRAGMA_ENABLE_DEPRECATION_WARNINGS
-
 bool FOpenXRInputPlugin::FOpenXRInput::SetEnhancedInputUserSettings(TObjectPtr<class UEnhancedInputUserSettings> InputSettings)
 {
 	if (bActionsAttached)
diff --git a/Engine/Plugins/Runtime/OpenXR/Source/OpenXRInput/Private/OpenXRInput.h b/Engine/Plugins/Runtime/OpenXR/Source/OpenXRInput/Private/OpenXRInput.h
--- a/Engine/Plugins/Runtime/OpenXR/Source/OpenXRInput/Private/OpenXRInput.h
+++ b/Engine/Plugins/Runtime/OpenXR/Source/OpenXRInput/Private/OpenXRInput.h
@@ -132,9 +132,6 @@ public:
 		virtual bool GetControllerOrientationAndPositionForTime(const int32 ControllerIndex, const FName MotionSource, FTimespan Time, bool& OutTimeWasUsed, FRotator& OutOrientation, FVector& OutPosition, bool& OutbProvidedLinearVelocity, FVector& OutLinearVelocity, bool& OutbProvidedAngularVelocity, FVector& OutAngularVelocityAsAxisAndLength, bool& OutbProvidedLinearAcceleration, FVector& OutLinearAcceleration, float WorldToMetersScale) const override;
 		virtual ETrackingStatus GetControllerTrackingStatus(const int32 ControllerIndex, const FName MotionSource) const override;
 		virtual void EnumerateSources(TArray<FMotionControllerSource>& SourcesOut) const override;
-		PRAGMA_DISABLE_DEPRECATION_WARNINGS
-		virtual bool SetPlayerMappableInputConfig(TObjectPtr<class UPlayerMappableInputConfig> InputConfig) override;
-		PRAGMA_ENABLE_DEPRECATION_WARNINGS
 		virtual bool SetEnhancedInputUserSettings(TObjectPtr<class UEnhancedInputUserSettings> InputSettings) override;
 		virtual bool AttachInputMappingContexts(const TSet<TObjectPtr<UInputMappingContext>>& MappingContexts) override;
 
diff --git a/Engine/Source/Runtime/HeadMountedDisplay/Public/IMotionController.h b/Engine/Source/Runtime/HeadMountedDisplay/Public/IMotionController.h
--- a/Engine/Source/Runtime/HeadMountedDisplay/Public/IMotionController.h
+++ b/Engine/Source/Runtime/HeadMountedDisplay/Public/IMotionController.h
@@ -151,19 +151,6 @@ public:
 	
 	virtual bool GetHandJointPosition(const FName MotionSource, int jointIndex, FVector& OutPosition) const = 0;
 
-	PRAGMA_DISABLE_DEPRECATION_WARNINGS
-	/**
-	 * Add a player mappable input config to the motion controller. This allows the motion controller to support
-	 * Enhanced Input actions.
-	 *
-	 * @param InputConfig		The path to the player mappable input config asset
-	 *
-	 * @return			False if the input config can't be attached to the session, true otherwise
-	 */
-	UE_DEPRECATED(5.4, "UPlayerMappableInputConfig has been deprecated. Please use SetInputMappingContexts instead.")
-	virtual bool SetPlayerMappableInputConfig(TObjectPtr<class UPlayerMappableInputConfig> InputConfig = nullptr) { return true; };
-	PRAGMA_ENABLE_DEPRECATION_WARNINGS
-
 	/**
 	* Add a set of enhanced input user settings to the motion controller. This allows the motion controller to support
 	* Enhanced Input actions.
