Support XR_EXT_user_presence
-Added support for the user presence extension, which tells us whether the hmd is being worn if the platform supplies that information.
-Removed the previous implementation which attempted to drive the worn state from other openxr state transitions.  This never really worked correctly and had inconsistent behavior across platforms.

diff --git a/Engine/Plugins/Runtime/OpenXR/Source/OpenXRHMD/Private/OpenXRHMD.cpp b/Engine/Plugins/Runtime/OpenXR/Source/OpenXRHMD/Private/OpenXRHMD.cpp
--- a/Engine/Plugins/Runtime/OpenXR/Source/OpenXRHMD/Private/OpenXRHMD.cpp
+++ b/Engine/Plugins/Runtime/OpenXR/Source/OpenXRHMD/Private/OpenXRHMD.cpp
@@ -2080,17 +2080,33 @@ bool FOpenXRHMD::AcquireSystemIdAndProperties()
 		return false;
 	}
 
-	// Retrieve system properties and check for hand tracking support
+	// Retrieve system properties and check for feature support
+	SystemProperties = XrSystemProperties{ XR_TYPE_SYSTEM_PROPERTIES };
+	void* Tail = &SystemProperties;
 	XrSystemHandTrackingPropertiesEXT HandTrackingSystemProperties = { XR_TYPE_SYSTEM_HAND_TRACKING_PROPERTIES_EXT };
-	SystemProperties = XrSystemProperties{ XR_TYPE_SYSTEM_PROPERTIES, &HandTrackingSystemProperties };
+	if (IsExtensionEnabled(XR_EXT_HAND_TRACKING_EXTENSION_NAME))
+	{
+		OpenXR::AppendChainStruct(Tail, &HandTrackingSystemProperties);
+	}
+	XrSystemUserPresencePropertiesEXT UserPresenceSystemProperties = { XR_TYPE_SYSTEM_USER_PRESENCE_PROPERTIES_EXT };
+	if (IsExtensionEnabled(XR_EXT_USER_PRESENCE_EXTENSION_NAME))
+	{
+		OpenXR::AppendChainStruct(Tail, &UserPresenceSystemProperties);
+	}
+
 	XrResult GetSystemPropsResult = xrGetSystemProperties(Instance, System, &SystemProperties);
 	XR_ENSURE(GetSystemPropsResult);
 	bxrGetSystemPropertiesSuccessful = (GetSystemPropsResult == XR_SUCCESS);
 
-	bSupportsHandTracking = HandTrackingSystemProperties.supportsHandTracking == XR_TRUE;
 	// Some runtimes aren't compliant with their number of layers supported.
 	// We support a fallback by emulating non-facelocked layers
 	bLayerSupportOpenXRCompliant = SystemProperties.graphicsProperties.maxLayerCount >= XR_MIN_COMPOSITION_LAYERS_SUPPORTED;
+	bSupportsHandTracking = HandTrackingSystemProperties.supportsHandTracking == XR_TRUE;
+	bUserPresenceSupported = UserPresenceSystemProperties.supportsUserPresence == XR_TRUE;
+
+	// We have built a linked list to function-scope structs and attached it to this member variable, those pointers will be bad after this function exits.
+	// Currently we never want to use them again, so we can just null out the next pointer.
+	SystemProperties.next = nullptr;
 
 	return true;
 }
@@ -3667,7 +3683,6 @@ bool FOpenXRHMD::OnStartGameFrame(FWorldContext& WorldContext)
 				{
 					GEngine->SetMaxFPS(0);
 				}
-				FCoreDelegates::VRHeadsetPutOnHead.Broadcast();
 				bIsReady = true;
 				StartSession();
 			}
@@ -3686,7 +3701,6 @@ bool FOpenXRHMD::OnStartGameFrame(FWorldContext& WorldContext)
 					const int32 PausedIdleFPS = CVarOpenXRPausedIdleFPS.GetValueOnAnyThread();
 					GEngine->SetMaxFPS(PausedIdleFPS);
 				}
-				FCoreDelegates::VRHeadsetRemovedFromHead.Broadcast();
 				bIsReady = false;
 				StopSession();
 			}
@@ -3740,6 +3754,7 @@ bool FOpenXRHMD::OnStartGameFrame(FWorldContext& WorldContext)
 		{
 			const XrEventDataReferenceSpaceChangePending& SpaceChange =
 				reinterpret_cast<XrEventDataReferenceSpaceChangePending&>(event);
+			check(SpaceChange.type == XR_TYPE_EVENT_DATA_REFERENCE_SPACE_CHANGE_PENDING);
 
 			if (SpaceChange.referenceSpaceType == XR_REFERENCE_SPACE_TYPE_STAGE)
 			{
@@ -3760,6 +3775,22 @@ bool FOpenXRHMD::OnStartGameFrame(FWorldContext& WorldContext)
 			bNeedReBuildOcclusionMesh = true;
 			break;
 		}
+		case XR_TYPE_EVENT_DATA_USER_PRESENCE_CHANGED_EXT:
+		{
+			const XrEventDataUserPresenceChangedEXT& PresenceChanged = 
+				reinterpret_cast<XrEventDataUserPresenceChangedEXT&>(event);
+			check(PresenceChanged.type == XR_TYPE_EVENT_DATA_USER_PRESENCE_CHANGED_EXT);
+
+			if (PresenceChanged.isUserPresent)
+			{
+				FCoreDelegates::VRHeadsetPutOnHead.Broadcast();
+			}
+			else
+			{
+				FCoreDelegates::VRHeadsetRemovedFromHead.Broadcast();
+			}
+			break;
+		}
 		}
 	}
 
diff --git a/Engine/Plugins/Runtime/OpenXR/Source/OpenXRHMD/Private/OpenXRHMD.h b/Engine/Plugins/Runtime/OpenXR/Source/OpenXRHMD/Private/OpenXRHMD.h
--- a/Engine/Plugins/Runtime/OpenXR/Source/OpenXRHMD/Private/OpenXRHMD.h
+++ b/Engine/Plugins/Runtime/OpenXR/Source/OpenXRHMD/Private/OpenXRHMD.h
@@ -605,6 +605,7 @@ private:
 	bool					bFoveationExtensionSupported;
 	bool					bRuntimeFoveationSupported;
 	bool					bLocalFloorSpaceSupported;
+	bool					bUserPresenceSupported;
 
 	XrColor4f				LayerColorScale;
 	XrColor4f				LayerColorBias;
diff --git a/Engine/Plugins/Runtime/OpenXR/Source/OpenXRHMD/Private/OpenXRHMDModule.cpp b/Engine/Plugins/Runtime/OpenXR/Source/OpenXRHMD/Private/OpenXRHMDModule.cpp
--- a/Engine/Plugins/Runtime/OpenXR/Source/OpenXRHMD/Private/OpenXRHMDModule.cpp
+++ b/Engine/Plugins/Runtime/OpenXR/Source/OpenXRHMD/Private/OpenXRHMDModule.cpp
@@ -589,6 +589,7 @@ bool FOpenXRHMDModule::GetOptionalExtensions(TArray<const ANSICHAR*>& OutExtensi
 	}
 	
 	OutExtensions.Add(XR_EXT_LOCAL_FLOOR_EXTENSION_NAME);
+	OutExtensions.Add(XR_EXT_USER_PRESENCE_EXTENSION_NAME);
 	return true;
 }
 
diff --git a/Engine/Plugins/Runtime/OpenXR/Source/OpenXRHMD/Public/OpenXRCore.h b/Engine/Plugins/Runtime/OpenXR/Source/OpenXRHMD/Public/OpenXRCore.h
--- a/Engine/Plugins/Runtime/OpenXR/Source/OpenXRHMD/Public/OpenXRCore.h
+++ b/Engine/Plugins/Runtime/OpenXR/Source/OpenXRHMD/Public/OpenXRCore.h
@@ -327,7 +327,7 @@ namespace OpenXR
 			}
 			else
 			{
-				Ptr = (XrBaseOutStructure*)Ptr->next;
+				Ptr = Ptr->next;
 			}
 		}
 		return nullptr;
@@ -345,9 +345,20 @@ namespace OpenXR
 			}
 			else
 			{
-				Ptr = (XrBaseInStructure*)Ptr->next;
+				Ptr = Ptr->next;
 			}
 		}
 		return nullptr;
 	}
-}
\ No newline at end of file
+
+	inline void AppendChainStruct(void*& Tail, void* NewTail)
+	{
+		XrBaseOutStructure*& TailXR = reinterpret_cast<XrBaseOutStructure*&>(Tail);
+		XrBaseOutStructure*& NewTailXR = reinterpret_cast<XrBaseOutStructure*&>(NewTail);
+		
+		check(TailXR->next == nullptr);
+		check(NewTailXR->next == nullptr);
+		TailXR->next = NewTailXR;
+		Tail = NewTailXR;
+	}
+}
