Add FOpenXRRenderBridge::RunOnRHISubmissionThread, and use it to wrap all OpenXR calls that can submit directly to the graphics queue (xrBeginFrame, xrEndFrame, xrAcquireSwapchainImage, xrReleaseSwapchainImage).
By default, this function will just run the provided lambda. It is overridden by FD3D12RenderBridge and FVulkanRenderBridge, which wrap the lambda in their respective RHIRunOnQueue functions, ensuring that submissions from the runtime are run on the submission thread if it is currently in use.

Also refactor OpenXRHMD thread checks to ensure code that may be wrapped in RHIRunOnQueue can run from any thread. The submission thread itself does not have an ETaskTag, and when the submission thread is disabled, code queued via RHIRunOnQueue can run from any thread that triggers a flush of submissions. This is manually synchronized and doesn't pose a concurrency risk, but it requires relaxing or moving the OpenXRHMD thread checks in some places.

diff --git a/Engine/Plugins/Runtime/OpenXR/Source/OpenXRHMD/Internal/OpenXRHMD_RenderBridge.h b/Engine/Plugins/Runtime/OpenXR/Source/OpenXRHMD/Internal/OpenXRHMD_RenderBridge.h
--- a/Engine/Plugins/Runtime/OpenXR/Source/OpenXRHMD/Internal/OpenXRHMD_RenderBridge.h
+++ b/Engine/Plugins/Runtime/OpenXR/Source/OpenXRHMD/Internal/OpenXRHMD_RenderBridge.h
@@ -47,6 +47,20 @@ public:
 			Template->GetClearBinding());
 	}
 
+	/** 
+	 * Used to wrap OpenXR calls which can submit directly to the graphics queue (xrBeginFrame, xrEndFrame, xrAcquireSwapchainImage, xrReleaseSwapchainImage) and their surrounding logic.
+	 * When a submission thread is in use, only that thread is permitted to submit to the graphics queue.
+	 * Implmentations for RHIs using a submission thread should override this function to run the provided code on the submission thread via RHIRunOnQueue.
+	 * 
+	 * Note: This is not necessary for calls in Present/HMDOnFinishRendering_RHIThread because FD3D12Viewport::Present and FVulkanViewport::Present 
+	 * both flush the submission thread before calling FRHICustomPresent::Present.
+	 */
+	virtual void RunOnRHISubmissionThread(TFunction<void()>&& CodeToRun)
+	{
+		ensure(IsInRenderingThread() || IsInRHIThread());
+		CodeToRun();
+	};
+
 	/** FRHICustomPresent */
 	virtual bool Present(IRHICommandContext& RHICmdContext, int32& InOutSyncInterval) override;
 
diff --git a/Engine/Plugins/Runtime/OpenXR/Source/OpenXRHMD/Private/OpenXRHMD.cpp b/Engine/Plugins/Runtime/OpenXR/Source/OpenXRHMD/Private/OpenXRHMD.cpp
--- a/Engine/Plugins/Runtime/OpenXR/Source/OpenXRHMD/Private/OpenXRHMD.cpp
+++ b/Engine/Plugins/Runtime/OpenXR/Source/OpenXRHMD/Private/OpenXRHMD.cpp
@@ -191,25 +191,6 @@ namespace {
 	#endif
 		return nullptr;
 	}
-
-	/** Hotfix 5.6.1: Helper function to run OpenXR calls which may access VkQueue on the Vulkan submission thread */
-	void RunOnGraphicsQueue(TFunction<void()>&& CodeToRun)
-	{
-#ifdef XR_USE_GRAPHICS_API_VULKAN
-		if (RHIGetInterfaceType() == ERHIInterfaceType::Vulkan)
-		{
-			GetIVulkanDynamicRHI()->RHIRunOnQueue(EVulkanRHIRunOnQueueType::Graphics, [CodeToRun](VkQueue Queue)
-			{
-				FTaskTagScope Scope(ETaskTag::EParallelRhiThread);
-				CodeToRun();
-			}, false);
-		}
-		else
-#endif
-		{
-			CodeToRun();
-		}
-	}
 }
 
 //---------------------------------------------------
@@ -2640,9 +2621,9 @@ bool FOpenXRHMD::AllocateRenderTargetTextures(uint32 SizeX, uint32 SizeY, uint8
 		if (!bIsAcquireOnAnyThreadSupported)
 		{
 			FRHICommandListImmediate& RHICmdList = FRHICommandListImmediate::Get();
-			RHICmdList.EnqueueLambda([Swapchain](FRHICommandListImmediate&)
+			RHICmdList.EnqueueLambda([this, Swapchain](FRHICommandListImmediate&)
 			{
-				RunOnGraphicsQueue([Swapchain]()
+				RenderBridge->RunOnRHISubmissionThread([Swapchain]()
 				{
 					Swapchain->IncrementSwapChainIndex_RHIThread();
 				});
@@ -2678,9 +2659,9 @@ bool FOpenXRHMD::AllocateRenderTargetTextures(uint32 SizeX, uint32 SizeY, uint8
 			if (!bIsAcquireOnAnyThreadSupported)
 			{
 				FRHICommandListImmediate& RHICmdList = FRHICommandListImmediate::Get();
-				RHICmdList.EnqueueLambda([EmulationSwapchain](FRHICommandListImmediate&)
+				RHICmdList.EnqueueLambda([this, EmulationSwapchain](FRHICommandListImmediate&)
 				{
-					RunOnGraphicsQueue([EmulationSwapchain]
+					RenderBridge->RunOnRHISubmissionThread([EmulationSwapchain]
 					{
 						EmulationSwapchain->IncrementSwapChainIndex_RHIThread();
 					});
@@ -2737,9 +2718,9 @@ void FOpenXRHMD::AllocateDepthTextureInternal(uint32 SizeX, uint32 SizeY, uint32
 		if (!bIsAcquireOnAnyThreadSupported)
 		{
 			FRHICommandListImmediate& RHICmdList = FRHICommandListImmediate::Get();
-			RHICmdList.EnqueueLambda([DepthSwapchain](FRHICommandListImmediate&)
+			RHICmdList.EnqueueLambda([this, DepthSwapchain](FRHICommandListImmediate&)
 			{
-				RunOnGraphicsQueue([DepthSwapchain]()
+				RenderBridge->RunOnRHISubmissionThread([DepthSwapchain]()
 				{
 					DepthSwapchain->IncrementSwapChainIndex_RHIThread();
 				});
@@ -3316,7 +3297,7 @@ void FOpenXRHMD::OnBeginRendering_RenderThread(FRDGBuilder& GraphBuilder, FScene
 		GraphBuilder.RHICmdList.EnqueueLambda([this, FrameState = PipelinedFrameStateRendering, ColorSwapchain, DepthSwapchain, EmulationSwapchain](FRHICommandListImmediate& InRHICmdList)
 		{
 			IRHICommandContext* Context = &InRHICmdList.GetContext();
-			RunOnGraphicsQueue([this, FrameState = PipelinedFrameStateRendering, ColorSwapchain, DepthSwapchain, EmulationSwapchain, Context]()
+			RenderBridge->RunOnRHISubmissionThread([this, FrameState = PipelinedFrameStateRendering, ColorSwapchain, DepthSwapchain, EmulationSwapchain, Context]()
 			{
 				OnBeginRendering_RHIThread(*Context, FrameState, ColorSwapchain, DepthSwapchain, EmulationSwapchain);
 			});
@@ -3782,8 +3763,6 @@ void FOpenXRHMD::RequestExitApp()
 
 void FOpenXRHMD::OnBeginRendering_RHIThread(IRHICommandContext& RHICmdContext, const FPipelinedFrameState& InFrameState, FXRSwapChainPtr ColorSwapchain, FXRSwapChainPtr DepthSwapchain, FXRSwapChainPtr EmulationSwapchain)
 {
-	ensure(IsInRenderingThread() || IsInRHIThread() || IsInParallelRHIThread());
-
 	// TODO: Add a hook to resolve discarded frames before we start a new frame.
 	UE_CLOG(bIsRendering, LogHMD, Verbose, TEXT("Discarded previous frame and started rendering a new frame."));
 
@@ -3893,7 +3872,7 @@ void FOpenXRHMD::OnBeginRendering_RHIThread(IRHICommandContext& RHICmdContext, c
 
 void FOpenXRHMD::OnFinishRendering_RHIThread(IRHICommandContext& RHICmdContext)
 {
-	ensure(IsInRenderingThread() || IsInRHIThread() || IsInParallelRHIThread());
+	check(IsInRenderingThread() || IsInRHIThread());
 
 	SCOPED_NAMED_EVENT(EndFrame, FColor::Red);
 
@@ -4130,11 +4109,11 @@ FIntRect FOpenXRHMD::GetFullFlatEyeRect_RenderThread(const FRHITextureDesc& EyeT
 void FOpenXRHMD::CopySwapchainTexture_RenderThread(FRDGBuilder& GraphBuilder, FRDGTextureRef SrcTexture, FIntRect SrcRect, const FXRSwapChainPtr& DstSwapChain,
 	FIntRect DstRect, bool bClearBlack, EXRCopyTextureBlendModifier SrcTextureCopyModifier, FStaticFeatureLevel FeatureLevel, FStaticShaderPlatform ShaderPlatform) const
 {
-	AddPass(GraphBuilder, RDG_EVENT_NAME("OpenXRHMD_AcquireLayerSwapchain"), [DstSwapChain](FRHICommandListImmediate& RHICmdList)
+	AddPass(GraphBuilder, RDG_EVENT_NAME("OpenXRHMD_AcquireLayerSwapchain"), [this, DstSwapChain](FRHICommandListImmediate& RHICmdList)
 	{
-		RHICmdList.EnqueueLambda([DstSwapChain](FRHICommandListImmediate& InRHICmdList)
+		RHICmdList.EnqueueLambda([this, DstSwapChain](FRHICommandListImmediate& InRHICmdList)
 		{
-			RunOnGraphicsQueue([DstSwapChain]()
+			RenderBridge->RunOnRHISubmissionThread([DstSwapChain]()
 			{
 				DstSwapChain->IncrementSwapChainIndex_RHIThread();
 				DstSwapChain->WaitCurrentImage_RHIThread(OPENXR_SWAPCHAIN_WAIT_TIMEOUT);
@@ -4156,11 +4135,11 @@ void FOpenXRHMD::CopySwapchainTexture_RenderThread(FRDGBuilder& GraphBuilder, FR
 
 	// Enqueue a command to release the image after the copy is done
 	bool bCapturableRequiresRHIContext = RuntimeRequiresRHIContext();
-	AddPass(GraphBuilder, RDG_EVENT_NAME("OpenXRHMD_ReleaseLayerSwapchain"), [DstSwapChain, bCapturableRequiresRHIContext](FRHICommandListImmediate& RHICmdList)
+	AddPass(GraphBuilder, RDG_EVENT_NAME("OpenXRHMD_ReleaseLayerSwapchain"), [this, DstSwapChain, bCapturableRequiresRHIContext](FRHICommandListImmediate& RHICmdList)
 	{
-		RHICmdList.EnqueueLambda([DstSwapChain, bCapturableRequiresRHIContext](FRHICommandListImmediate& InRHICmdList)
+		RHICmdList.EnqueueLambda([this, DstSwapChain, bCapturableRequiresRHIContext](FRHICommandListImmediate& InRHICmdList)
 		{
-			RunOnGraphicsQueue([DstSwapChain, bCapturableRequiresRHIContext, &InRHICmdList]()
+			RenderBridge->RunOnRHISubmissionThread([DstSwapChain, bCapturableRequiresRHIContext, &InRHICmdList]()
 			{
 				DstSwapChain->ReleaseCurrentImage_RHIThread(bCapturableRequiresRHIContext ? &InRHICmdList.GetContext() : nullptr);
 			});
diff --git a/Engine/Plugins/Runtime/OpenXR/Source/OpenXRHMD/Private/OpenXRHMD_RenderBridge.cpp b/Engine/Plugins/Runtime/OpenXR/Source/OpenXRHMD/Private/OpenXRHMD_RenderBridge.cpp
--- a/Engine/Plugins/Runtime/OpenXR/Source/OpenXRHMD/Private/OpenXRHMD_RenderBridge.cpp
+++ b/Engine/Plugins/Runtime/OpenXR/Source/OpenXRHMD/Private/OpenXRHMD_RenderBridge.cpp
@@ -122,6 +122,12 @@ public:
 		return CreateSwapchain_D3D12(InSession, Format, OutActualFormat, SizeX, SizeY, ArraySize, NumMips, NumSamples, CreateFlags, ClearValueBinding, AuxiliaryCreateFlags);
 	}
 
+	virtual void RunOnRHISubmissionThread(TFunction<void()>&& CodeToRun) override
+	{
+		ensure(IsInRenderingThread() || IsInRHIThread());
+		GetID3D12DynamicRHI()->RHIRunOnQueue(ED3D12RHIRunOnQueueType::Graphics, [CodeToRun](ID3D12CommandQueue* Queue) { CodeToRun(); }, false);
+	}
+
 private:
 	PFN_xrGetD3D12GraphicsRequirementsKHR GetD3D12GraphicsRequirementsKHR;
 	XrGraphicsBindingD3D12KHR Binding;
@@ -317,6 +323,12 @@ public:
 		return CreateSwapchain_Vulkan(InSession, Format, OutActualFormat, SizeX, SizeY, ArraySize, NumMips, NumSamples, CreateFlags, ClearValueBinding, AuxiliaryCreateFlags);
 	}
 
+	virtual void RunOnRHISubmissionThread(TFunction<void()>&& CodeToRun) override
+	{
+		ensure(IsInRenderingThread() || IsInRHIThread());
+		GetIVulkanDynamicRHI()->RHIRunOnQueue(EVulkanRHIRunOnQueueType::Graphics, [CodeToRun](VkQueue Queue) { CodeToRun(); }, false);
+	}
+
 private:
 	PFN_xrGetVulkanGraphicsRequirementsKHR GetVulkanGraphicsRequirementsKHR;
 	PFN_xrGetVulkanGraphicsDeviceKHR GetVulkanGraphicsDeviceKHR;
diff --git a/Engine/Plugins/Runtime/OpenXR/Source/OpenXRHMD/Private/OpenXRHMD_Swapchain.cpp b/Engine/Plugins/Runtime/OpenXR/Source/OpenXRHMD/Private/OpenXRHMD_Swapchain.cpp
--- a/Engine/Plugins/Runtime/OpenXR/Source/OpenXRHMD/Private/OpenXRHMD_Swapchain.cpp
+++ b/Engine/Plugins/Runtime/OpenXR/Source/OpenXRHMD/Private/OpenXRHMD_Swapchain.cpp
@@ -52,8 +52,6 @@ void FOpenXRSwapchain::IncrementSwapChainIndex_RHIThread()
 
 void FOpenXRSwapchain::WaitCurrentImage_RHIThread(int64 Timeout)
 {
-	check(IsInRenderingThread() || IsInRHIThread() || IsInParallelRHIThread());
-
 	bool WasAcquired = true;
 	ImageAcquired.compare_exchange_strong(WasAcquired, false);
 	if (!WasAcquired)
@@ -99,8 +97,6 @@ void FOpenXRSwapchain::WaitCurrentImage_RHIThread(int64 Timeout)
 
 void FOpenXRSwapchain::ReleaseCurrentImage_RHIThread(IRHICommandContext* RHICmdContext)
 {
-	check(IsInRenderingThread() || IsInRHIThread() || IsInParallelRHIThread());
-
 	bool WasReady = true;
 	ImageReady.compare_exchange_strong(WasReady, false);
 	if (!WasReady)
