Improve OpenXR call ensure
-Improved XR_ENSURE so that the ensure condition and the ensure message contain the xrFunction call which is being checked.
-Added XR_ENSURE_WITH_CALLINFO so that in cases where you need to call the function before the ensure you can provide the name of the function being checked, and that will appear in the ensure msg.
-Fixed up a few specific cases where we would get less information.

diff --git a/Engine/Plugins/Runtime/OpenXR/Source/OpenXRHMD/Private/OpenXRCore.cpp b/Engine/Plugins/Runtime/OpenXR/Source/OpenXRHMD/Private/OpenXRCore.cpp
--- a/Engine/Plugins/Runtime/OpenXR/Source/OpenXRHMD/Private/OpenXRCore.cpp
+++ b/Engine/Plugins/Runtime/OpenXR/Source/OpenXRHMD/Private/OpenXRCore.cpp
@@ -20,8 +20,7 @@ FOpenXRPath::FOpenXRPath(const char* PathString)
 	XrInstance Instance = IOpenXRHMDModule::Get().GetInstance();
 	if (ensure(Instance != XR_NULL_HANDLE))
 	{
-		XrResult Result = xrStringToPath(Instance, PathString, &Path);
-		check(XR_SUCCEEDED(Result));
+		XR_ENSURE(xrStringToPath(Instance, PathString, &Path));
 	}
 }
 
diff --git a/Engine/Plugins/Runtime/OpenXR/Source/OpenXRHMD/Private/OpenXRHMD.cpp b/Engine/Plugins/Runtime/OpenXR/Source/OpenXRHMD/Private/OpenXRHMD.cpp
--- a/Engine/Plugins/Runtime/OpenXR/Source/OpenXRHMD/Private/OpenXRHMD.cpp
+++ b/Engine/Plugins/Runtime/OpenXR/Source/OpenXRHMD/Private/OpenXRHMD.cpp
@@ -2259,7 +2259,7 @@ bool FOpenXRHMD::OnStereoTeardown()
 	}
 	else
 	{
-		XR_ENSURE(Result);
+		XR_ENSURE_WITH_CALLINFO(Result, "xrRequestExitSession");
 	}
 
 	FCoreDelegates::VRHeadsetRecenter.RemoveAll(this);
@@ -3492,9 +3492,9 @@ bool FOpenXRHMD::ReadNextEvent(XrEventDataBuffer* buffer)
 	// It is sufficient to clear just the XrEventDataBuffer header to XR_TYPE_EVENT_DATA_BUFFER
 	XrEventDataBaseHeader* baseHeader = reinterpret_cast<XrEventDataBaseHeader*>(buffer);
 	*baseHeader = { XR_TYPE_EVENT_DATA_BUFFER };
-	const XrResult xr = xrPollEvent(Instance, buffer);
-	XR_ENSURE(xr);
-	if (xr == XR_SUCCESS)
+	XrResult Result;
+	XR_ENSURE(Result = xrPollEvent(Instance, buffer));
+	if (Result == XR_SUCCESS)
 	{
 		for (IOpenXRExtensionPlugin* Module : ExtensionPlugins)
 		{
diff --git a/Engine/Plugins/Runtime/OpenXR/Source/OpenXRHMD/Public/OpenXRCore.h b/Engine/Plugins/Runtime/OpenXR/Source/OpenXRHMD/Public/OpenXRCore.h
--- a/Engine/Plugins/Runtime/OpenXR/Source/OpenXRHMD/Public/OpenXRCore.h
+++ b/Engine/Plugins/Runtime/OpenXR/Source/OpenXRHMD/Public/OpenXRCore.h
@@ -37,13 +37,36 @@ constexpr const TCHAR* OpenXRReferenceSpaceTypeToString(XrReferenceSpaceType e)
 	}
 }
 
+// In order to get the parameters that are in x (whatever they are) I need to implicitly capture everything.  This is quite dangerous in a lambda that is going to 
+// be passed off somewhere, but this lambda is being immediately used inline.
+// I find this a bit obtuse to read, so what this does:
+// Define a macro with one parameter, intended to be an openxr api call and parameters, eg: "xrStringToPath(Instance, PathString, &Path)".
+// Define a lamba with no parameters, implicit capture by reference, returning a bool.
+// In the body of the lambda ensureMsgf with the condition of the ensure calling the function and assigning its return to XRES_Result then use the comma operator and test the result.  This means that the function is in the
+// condition of the ensure, which appears in the log before the message. Because of the implicit capture the paramters of x are available.  Then we build the message of the ensure including both x and the XRES_Result. Then return the result of the ensure.
+// Call this lambda immediately, thats the last two parentheses.
+// Sample usage: if (XR_ENSURE(xrStringToPath(Instance, PathString, &Path))) {...}
+// Sample usage2: XrResult ResultForLater; if (XR_ENSURE(ResultForLater = xrStringToPath(Instance, PathString, &Path))) {...}
 #if DO_CHECK
-#define XR_ENSURE(x) [] (XrResult Result) \
+#define XR_ENSURE(xrCall) \
+	[&] () -> bool \
 	{ \
-		return ensureMsgf(XR_SUCCEEDED(Result), TEXT("OpenXR call failed with result %s"), OpenXRResultToString(Result)); \
-	} (x)
+		XrResult XRES_Result; \
+		return ensureMsgf((XRES_Result = xrCall, XR_SUCCEEDED(XRES_Result)), TEXT("XR call %hs failed with result: %s"), #xrCall, OpenXRResultToString(XRES_Result)); \
+	} ()
 #else
-#define XR_ENSURE(x) XR_SUCCEEDED(x)
+#define XR_ENSURE(xrCall) XR_SUCCEEDED(xrCall)
+#endif
+
+// To handle cases where we have an XrResult as a variable, and want an ensure for a function we called earlier.  Must manually provide call info text.
+// Sample usage: bool Success = XR_ENSURE_RESULT(Result, "xrRequestExitSession");
+#if DO_CHECK
+#define XR_ENSURE_WITH_CALLINFO(xrCall, xrCallInfoText) [] (XrResult Result) \
+	{ \
+		return ensureMsgf(XR_SUCCEEDED(Result), TEXT("XR call %hs failed with result %s"), #xrCallInfoText, OpenXRResultToString(Result)); \
+	} (xrCall)
+#else
+#define XR_ENSURE_WITH_CALLINFO(xrCall, xrCallInfoText) XR_SUCCEEDED(xrCall)
 #endif
 
 inline FQuat ToFQuat(XrQuaternionf Quat)
@@ -199,6 +222,8 @@ public:
 	 */
 	FName ToName() const;
 
+	XrPath ToXRPath() const { return Path; }
+
 	operator bool() const { return Path != XR_NULL_PATH; }
 	operator XrPath() const { return Path; }
 	operator FString() const { return ToString(); }
diff --git a/Engine/Plugins/Runtime/OpenXR/Source/OpenXRInput/Private/OpenXRInput.cpp b/Engine/Plugins/Runtime/OpenXR/Source/OpenXRInput/Private/OpenXRInput.cpp
--- a/Engine/Plugins/Runtime/OpenXR/Source/OpenXRInput/Private/OpenXRInput.cpp
+++ b/Engine/Plugins/Runtime/OpenXR/Source/OpenXRInput/Private/OpenXRInput.cpp
@@ -81,10 +81,10 @@ FOpenXRInputPlugin::~FOpenXRInputPlugin()
 
 void FOpenXRInputPlugin::StartupModule()
 {
-	IOpenXRInputPlugin::StartupModule();
-
 	// Note: XRSystem may be null, for example in the editor.  But we still need the input device to enumerate sources.
 	InputDevice = MakeShared<FOpenXRInput>(GEngine->XRSystem.Get());
+
+	IOpenXRInputPlugin::StartupModule();
 }
 
 FOpenXRInputPlugin::FOpenXRAction::FOpenXRAction(XrActionSet InActionSet,
@@ -385,12 +385,25 @@ bool FOpenXRInputPlugin::FOpenXRInput::BuildActions(XrSession Session)
 
 	// Generate a map of all supported interaction profiles to store suggested bindings
 	TMap<FString, FInteractionProfile> Profiles;
-	Profiles.Add("SimpleController", FInteractionProfile(FOpenXRPath("/interaction_profiles/khr/simple_controller"), true));
-	Profiles.Add("Vive", FInteractionProfile(FOpenXRPath("/interaction_profiles/htc/vive_controller"), true));
-	Profiles.Add("MixedReality", FInteractionProfile(FOpenXRPath("/interaction_profiles/microsoft/motion_controller"), true));
-	Profiles.Add("OculusGo", FInteractionProfile(FOpenXRPath("/interaction_profiles/oculus/go_controller"), false));
-	Profiles.Add("OculusTouch", FInteractionProfile(FOpenXRPath("/interaction_profiles/oculus/touch_controller"), true));
-	Profiles.Add("ValveIndex", FInteractionProfile(FOpenXRPath("/interaction_profiles/valve/index_controller"), true));
+	{
+		auto AddProfile = [&Profiles](const char* const Profile, const char* const PathChars, bool HasHaptics)
+			{
+				FOpenXRPath Path = FOpenXRPath(PathChars);
+				if (XR_NULL_PATH == Path.ToXRPath())
+				{
+					UE_LOG(LogHMD, Warning, TEXT("BuildActions tried to AddProfile: '%s' Path: '%s', but it looks like the path creation failed. Profile not added. Look above for path related error messages."), ANSI_TO_TCHAR(Profile), ANSI_TO_TCHAR(PathChars));
+					return;
+				}
+				Profiles.Add(Profile, FInteractionProfile(Path, true));
+			};
+		//AddProfile("MalformedPathTest", "MALFORMEDPATH *!,/", false); // malformed path, should fail with XR_ERROR_PATH_FORMAT_INVALID = -21
+		AddProfile("SimpleController", "/interaction_profiles/khr/simple_controller", true);
+		AddProfile("Vive", "/interaction_profiles/htc/vive_controller", true);
+		AddProfile("MixedReality", "/interaction_profiles/microsoft/motion_controller", true);
+		AddProfile("OculusGo", "/interaction_profiles/oculus/go_controller", false);
+		AddProfile("OculusTouch", "/interaction_profiles/oculus/touch_controller", true);
+		AddProfile("ValveIndex", "/interaction_profiles/valve/index_controller", true);
+	}
 
 	// Query extension plugins for input key overrides
 	for (IOpenXRExtensionPlugin* Plugin : OpenXRHMD->GetExtensionPlugins())
