Improve OpenXR call ensure
-Fixed AddProfile always being given true rather than the passed in bHasHaptics parameter.
-Fixed false-positive static analysis warning.  Improved sample usage comment to help avoid this problem.

diff --git a/Engine/Plugins/Runtime/OpenXR/Source/OpenXRHMD/Private/OpenXRHMD.cpp b/Engine/Plugins/Runtime/OpenXR/Source/OpenXRHMD/Private/OpenXRHMD.cpp
--- a/Engine/Plugins/Runtime/OpenXR/Source/OpenXRHMD/Private/OpenXRHMD.cpp
+++ b/Engine/Plugins/Runtime/OpenXR/Source/OpenXRHMD/Private/OpenXRHMD.cpp
@@ -3492,7 +3492,7 @@ bool FOpenXRHMD::ReadNextEvent(XrEventDataBuffer* buffer)
 	// It is sufficient to clear just the XrEventDataBuffer header to XR_TYPE_EVENT_DATA_BUFFER
 	XrEventDataBaseHeader* baseHeader = reinterpret_cast<XrEventDataBaseHeader*>(buffer);
 	*baseHeader = { XR_TYPE_EVENT_DATA_BUFFER };
-	XrResult Result;
+	XrResult Result = XR_ERROR_VALIDATION_FAILURE;
 	XR_ENSURE(Result = xrPollEvent(Instance, buffer));
 	if (Result == XR_SUCCESS)
 	{
diff --git a/Engine/Plugins/Runtime/OpenXR/Source/OpenXRHMD/Public/OpenXRCore.h b/Engine/Plugins/Runtime/OpenXR/Source/OpenXRHMD/Public/OpenXRCore.h
--- a/Engine/Plugins/Runtime/OpenXR/Source/OpenXRHMD/Public/OpenXRCore.h
+++ b/Engine/Plugins/Runtime/OpenXR/Source/OpenXRHMD/Public/OpenXRCore.h
@@ -46,7 +46,7 @@ constexpr const TCHAR* OpenXRReferenceSpaceTypeToString(XrReferenceSpaceType e)
 // condition of the ensure, which appears in the log before the message. Because of the implicit capture the paramters of x are available.  Then we build the message of the ensure including both x and the XRES_Result. Then return the result of the ensure.
 // Call this lambda immediately, thats the last two parentheses.
 // Sample usage: if (XR_ENSURE(xrStringToPath(Instance, PathString, &Path))) {...}
-// Sample usage2: XrResult ResultForLater; if (XR_ENSURE(ResultForLater = xrStringToPath(Instance, PathString, &Path))) {...}
+// Sample usage2: XrResult ResultForLater = XR_ERROR_VALIDATION_FAILURE; if (XR_ENSURE(ResultForLater = xrStringToPath(Instance, PathString, &Path))) {...}
 #if DO_CHECK
 #define XR_ENSURE(xrCall) \
 	[&] () -> bool \
diff --git a/Engine/Plugins/Runtime/OpenXR/Source/OpenXRInput/Private/OpenXRInput.cpp b/Engine/Plugins/Runtime/OpenXR/Source/OpenXRInput/Private/OpenXRInput.cpp
--- a/Engine/Plugins/Runtime/OpenXR/Source/OpenXRInput/Private/OpenXRInput.cpp
+++ b/Engine/Plugins/Runtime/OpenXR/Source/OpenXRInput/Private/OpenXRInput.cpp
@@ -394,7 +394,7 @@ bool FOpenXRInputPlugin::FOpenXRInput::BuildActions(XrSession Session)
 					UE_LOG(LogHMD, Warning, TEXT("BuildActions tried to AddProfile: '%s' Path: '%s', but it looks like the path creation failed. Profile not added. Look above for path related error messages."), ANSI_TO_TCHAR(Profile), ANSI_TO_TCHAR(PathChars));
 					return;
 				}
-				Profiles.Add(Profile, FInteractionProfile(Path, true));
+				Profiles.Add(Profile, FInteractionProfile(Path, HasHaptics));
 			};
 		//AddProfile("MalformedPathTest", "MALFORMEDPATH *!,/", false); // malformed path, should fail with XR_ERROR_PATH_FORMAT_INVALID = -21
 		AddProfile("SimpleController", "/interaction_profiles/khr/simple_controller", true);
