Improve OpenXR call ensure
-Suppressing a new static analysis warning.  The macro is confusing it.

diff --git a/Engine/Plugins/Runtime/OpenXR/Source/OpenXRHMD/Private/OpenXRHMD.cpp b/Engine/Plugins/Runtime/OpenXR/Source/OpenXRHMD/Private/OpenXRHMD.cpp
--- a/Engine/Plugins/Runtime/OpenXR/Source/OpenXRHMD/Private/OpenXRHMD.cpp
+++ b/Engine/Plugins/Runtime/OpenXR/Source/OpenXRHMD/Private/OpenXRHMD.cpp
@@ -3494,7 +3494,7 @@ bool FOpenXRHMD::ReadNextEvent(XrEventDataBuffer* buffer)
 	*baseHeader = { XR_TYPE_EVENT_DATA_BUFFER };
 	XrResult Result = XR_ERROR_VALIDATION_FAILURE;
 	XR_ENSURE(Result = xrPollEvent(Instance, buffer));
-	if (Result == XR_SUCCESS)
+	if (Result == XR_SUCCESS) //-V547  The macro above confuses static analasis into thinking this condition is always false.
 	{
 		for (IOpenXRExtensionPlugin* Module : ExtensionPlugins)
 		{
diff --git a/Engine/Plugins/Runtime/OpenXR/Source/OpenXRHMD/Public/OpenXRCore.h b/Engine/Plugins/Runtime/OpenXR/Source/OpenXRHMD/Public/OpenXRCore.h
--- a/Engine/Plugins/Runtime/OpenXR/Source/OpenXRHMD/Public/OpenXRCore.h
+++ b/Engine/Plugins/Runtime/OpenXR/Source/OpenXRHMD/Public/OpenXRCore.h
@@ -45,8 +45,9 @@ constexpr const TCHAR* OpenXRReferenceSpaceTypeToString(XrReferenceSpaceType e)
 // In the body of the lambda ensureMsgf with the condition of the ensure calling the function and assigning its return to XRES_Result then use the comma operator and test the result.  This means that the function is in the
 // condition of the ensure, which appears in the log before the message. Because of the implicit capture the paramters of x are available.  Then we build the message of the ensure including both x and the XRES_Result. Then return the result of the ensure.
 // Call this lambda immediately, thats the last two parentheses.
-// Sample usage: if (XR_ENSURE(xrStringToPath(Instance, PathString, &Path))) {...}
-// Sample usage2: XrResult ResultForLater = XR_ERROR_VALIDATION_FAILURE; if (XR_ENSURE(ResultForLater = xrStringToPath(Instance, PathString, &Path))) {...}
+// Sample Usage: if (XR_ENSURE(xrStringToPath(Instance, PathString, &Path))) {...}
+// Sample Usage 2: XrResult ResultForLater = XR_ERROR_VALIDATION_FAILURE; if (XR_ENSURE(ResultForLater = xrStringToPath(Instance, PathString, &Path))) {...}
+// Note: in Sample Usage 2 we initialize ResultForLater and we might need to suppress static analysis with //-V547 where we use ResultForLater.
 #if DO_CHECK
 #define XR_ENSURE(xrCall) \
 	[&] () -> bool \
