Add new blueprint nodes to control OpenXR passthrough settings:
- SetEnvironmentBlendMode() <-- pre-existing, but now uses the new UENUM
- GetEnvironmentBlendMode()
- GetSupportedEnvironmentBlendModes() <-- returns in order of runtime preference (highest to lowest)
- IsCompositionLayerInvertedAlphaEnabled()

Add xr.OpenXRInvertAlpha to the OpenXR plugin settings, and r.Mobile.PropagateAlpha to the Engine->Rendering->Mobile settings.
Rename OpenXRHMDSettings to OpenXR Settings, since it controls settings for the entire module.
Move some functions from FOpenXRHMD to IOpenXRHMD (slated for 5.7 after the 5.6 OpenXRHMD overhaul).

diff --git a/Engine/Plugins/Runtime/OpenXR/Source/OpenXREditor/Private/OpenXREditorModule.cpp b/Engine/Plugins/Runtime/OpenXR/Source/OpenXREditor/Private/OpenXREditorModule.cpp
--- a/Engine/Plugins/Runtime/OpenXR/Source/OpenXREditor/Private/OpenXREditorModule.cpp
+++ b/Engine/Plugins/Runtime/OpenXR/Source/OpenXREditor/Private/OpenXREditorModule.cpp
@@ -15,18 +15,6 @@ void FOpenXREditorModule::StartupModule()
 {
 	FOpenXRAssetDirectory::LoadForCook();
 
-	// register settings
-	ISettingsModule* SettingsModule = FModuleManager::GetModulePtr<ISettingsModule>("Settings");
-
-	if (SettingsModule != nullptr)
-	{
-		SettingsModule->RegisterSettings("Project", "Plugins", "OpenXRHMD",
-			LOCTEXT("OpenXRHMDSettingsName", "OpenXR Settings"),
-			LOCTEXT("OpenXRHMDSettingsDescription", "Project settings for OpenXR plugin"),
-			GetMutableDefault<UOpenXRHMDSettings>()
-		);
-	}
-
 	UEditorPerformanceSettings* EditorPerformanceSettings = GetMutableDefault<UEditorPerformanceSettings>();
 	if (EditorPerformanceSettings->bOverrideMaxViewportRenderingResolution)
 	{
diff --git a/Engine/Plugins/Runtime/OpenXR/Source/OpenXRHMD/OpenXRHMD.Build.cs b/Engine/Plugins/Runtime/OpenXR/Source/OpenXRHMD/OpenXRHMD.Build.cs
--- a/Engine/Plugins/Runtime/OpenXR/Source/OpenXRHMD/OpenXRHMD.Build.cs
+++ b/Engine/Plugins/Runtime/OpenXR/Source/OpenXRHMD/OpenXRHMD.Build.cs
@@ -37,6 +37,7 @@ namespace UnrealBuildTool.Rules
                     "SlateCore",
 					"AugmentedReality",
 					"EngineSettings",
+					"DeveloperSettings",
 				}
 				);
 
diff --git a/Engine/Plugins/Runtime/OpenXR/Source/OpenXRHMD/Private/OpenXRBlueprintFunctionLibrary.cpp b/Engine/Plugins/Runtime/OpenXR/Source/OpenXRHMD/Private/OpenXRBlueprintFunctionLibrary.cpp
--- a/Engine/Plugins/Runtime/OpenXR/Source/OpenXRHMD/Private/OpenXRBlueprintFunctionLibrary.cpp
+++ b/Engine/Plugins/Runtime/OpenXR/Source/OpenXRHMD/Private/OpenXRBlueprintFunctionLibrary.cpp
@@ -12,13 +12,91 @@ UOpenXRBlueprintFunctionLibrary::UOpenXRBlueprintFunctionLibrary(FObjectInitiali
 {
 }
 
-void UOpenXRBlueprintFunctionLibrary::SetEnvironmentBlendMode(int32 NewBlendMode) 
+void UOpenXRBlueprintFunctionLibrary::SetEnvironmentBlendMode(EOpenXREnvironmentBlendMode NewBlendMode) 
 {
+	if (NewBlendMode == EOpenXREnvironmentBlendMode::None)
+	{
+		UE_LOG(LogUOpenXRBlueprintFunctionLibrary, Warning, TEXT("Attempted to set environment blend mode to EOpenXREnvironmentBlendMode::None. No blend mode will be set and the runtime will continue to use its current blend mode."));
+		return;
+	}
+
+	FOpenXRHMD* OpenXRHMD = GetOpenXRHMD();
+	if (OpenXRHMD) 
+	{
+		OpenXRHMD->SetEnvironmentBlendMode(static_cast<XrEnvironmentBlendMode>(NewBlendMode));
+	}
+}
+
+EOpenXREnvironmentBlendMode UOpenXRBlueprintFunctionLibrary::GetEnvironmentBlendMode() 
+{
+	XrEnvironmentBlendMode CurrentBlendMode = XR_ENVIRONMENT_BLEND_MODE_MAX_ENUM;
+
 	FOpenXRHMD* OpenXRHMD = GetOpenXRHMD();
 	if (OpenXRHMD) 
 	{
-		OpenXRHMD->SetEnvironmentBlendMode((XrEnvironmentBlendMode) NewBlendMode);
+		CurrentBlendMode = OpenXRHMD->GetEnvironmentBlendMode();
 	}
+
+	switch (CurrentBlendMode)
+	{
+	case XR_ENVIRONMENT_BLEND_MODE_OPAQUE:
+	case XR_ENVIRONMENT_BLEND_MODE_ADDITIVE:
+	case XR_ENVIRONMENT_BLEND_MODE_ALPHA_BLEND:
+		return static_cast<EOpenXREnvironmentBlendMode>(CurrentBlendMode);
+
+	case XR_ENVIRONMENT_BLEND_MODE_MAX_ENUM:
+		return EOpenXREnvironmentBlendMode::None; // HMD is uninitialized
+
+	default:
+		checkNoEntry();
+		return EOpenXREnvironmentBlendMode::None;
+	}
+}
+
+TArray<EOpenXREnvironmentBlendMode> UOpenXRBlueprintFunctionLibrary::GetSupportedEnvironmentBlendModes() 
+{
+	TArray<XrEnvironmentBlendMode> SupportedBlendModes = {};
+	TArray<EOpenXREnvironmentBlendMode> OutputSupportedBlendModes = {};
+
+	FOpenXRHMD* OpenXRHMD = GetOpenXRHMD();
+	if (OpenXRHMD) 
+	{
+		SupportedBlendModes = OpenXRHMD->GetSupportedEnvironmentBlendModes(); // These should be ordered from highest to lowest runtime preference
+	}
+	
+	for (XrEnvironmentBlendMode BlendMode : SupportedBlendModes)
+	{
+		switch (BlendMode)
+		{
+		case XR_ENVIRONMENT_BLEND_MODE_OPAQUE:
+		case XR_ENVIRONMENT_BLEND_MODE_ADDITIVE:
+		case XR_ENVIRONMENT_BLEND_MODE_ALPHA_BLEND:
+			OutputSupportedBlendModes.Add(static_cast<EOpenXREnvironmentBlendMode>(BlendMode));
+			break;
+
+		case XR_ENVIRONMENT_BLEND_MODE_MAX_ENUM:
+		default:
+			checkNoEntry();
+			break;
+		}
+	}
+	
+	return OutputSupportedBlendModes;
+}
+
+bool UOpenXRBlueprintFunctionLibrary::IsCompositionLayerInvertedAlphaEnabled() 
+{
+	FOpenXRHMD* OpenXRHMD = GetOpenXRHMD();
+	if (OpenXRHMD) 
+	{
+		static const TConsoleVariableData<bool>* OpenXRInvertAlphaCVar = IConsoleManager::Get().FindTConsoleVariableDataBool(TEXT("xr.OpenXRInvertAlpha"));
+		const bool bEXTExtensionEnabled = OpenXRHMD->IsExtensionEnabled(XR_EXT_COMPOSITION_LAYER_INVERTED_ALPHA_EXTENSION_NAME);
+		const bool bFBExtensionEnabled = OpenXRHMD->IsExtensionEnabled(XR_FB_COMPOSITION_LAYER_ALPHA_BLEND_EXTENSION_NAME);
+		
+		return OpenXRInvertAlphaCVar->GetValueOnGameThread() && (bEXTExtensionEnabled || bFBExtensionEnabled);
+	}
+
+	return false;
 }
 
 FOpenXRHMD* UOpenXRBlueprintFunctionLibrary::GetOpenXRHMD()
diff --git a/Engine/Plugins/Runtime/OpenXR/Source/OpenXRHMD/Private/OpenXRHMD.cpp b/Engine/Plugins/Runtime/OpenXR/Source/OpenXRHMD/Private/OpenXRHMD.cpp
--- a/Engine/Plugins/Runtime/OpenXR/Source/OpenXRHMD/Private/OpenXRHMD.cpp
+++ b/Engine/Plugins/Runtime/OpenXR/Source/OpenXRHMD/Private/OpenXRHMD.cpp
@@ -1638,7 +1638,7 @@ bool FOpenXRHMD::ReconfigureForShaderPlatform(EShaderPlatform NewShaderPlatform)
 	return true;
 }
 
-TArray<XrEnvironmentBlendMode> FOpenXRHMD::RetrieveEnvironmentBlendModes() const
+TArray<XrEnvironmentBlendMode> FOpenXRHMD::GetSupportedEnvironmentBlendModes() const
 {
 	TArray<XrEnvironmentBlendMode> BlendModes;
 	uint32 BlendModeCount;
@@ -2067,7 +2067,7 @@ bool FOpenXRHMD::OnStereoStartup()
 
 	// Select the first blend mode returned by the runtime - as per spec, environment blend modes should be in order from highest to lowest runtime preference
 	{
-		TArray<XrEnvironmentBlendMode> BlendModes = RetrieveEnvironmentBlendModes();
+		TArray<XrEnvironmentBlendMode> BlendModes = GetSupportedEnvironmentBlendModes();
 		if (!BlendModes.IsEmpty())
 		{
 			SelectedEnvironmentBlendMode = BlendModes[0];
@@ -2451,7 +2451,7 @@ void FOpenXRHMD::SetEnvironmentBlendMode(XrEnvironmentBlendMode NewBlendMode)
 		return;
 	}
 
-	TArray<XrEnvironmentBlendMode> BlendModes = RetrieveEnvironmentBlendModes();
+	TArray<XrEnvironmentBlendMode> BlendModes = GetSupportedEnvironmentBlendModes();
 
 	if (BlendModes.Contains(NewBlendMode))
 	{
@@ -2501,7 +2501,6 @@ bool FOpenXRHMD::StopSession()
 void FOpenXRHMD::OnBeginPlay(FWorldContext& InWorldContext)
 {
 	bOpenXRForceStereoLayersEmulationCVarCachedValue = CVarOpenXRForceStereoLayerEmulation.GetValueOnGameThread();
-	bOpenXRInvertAlphaCvarCachedValue = CVarOpenXRInvertAlpha.GetValueOnGameThread();
 
 	const UOpenXRHMDSettings* Settings = GetDefault<UOpenXRHMDSettings>();
 	bRuntimeFoveationSupported = bFoveationExtensionSupported && (Settings != nullptr ? Settings->bIsFBFoveationEnabled : false);
@@ -3917,7 +3916,7 @@ void FOpenXRHMD::OnFinishRendering_RHIThread(IRHICommandContext& RHICmdContext)
 			Layer.views = PipelinedLayerStateRHI.ProjectionLayers.GetData();
 			Headers.Add(reinterpret_cast<XrCompositionLayerBaseHeader*>(&Layer));
 
-			if (bOpenXRInvertAlphaCvarCachedValue)
+			if (CVarOpenXRInvertAlpha.GetValueOnRenderThread())
 			{
 				// These two extensions do the same thing, the first one is more modern.  The second is to keep older runtimes functioning.
 				if(IsExtensionEnabled(XR_EXT_COMPOSITION_LAYER_INVERTED_ALPHA_EXTENSION_NAME))
diff --git a/Engine/Plugins/Runtime/OpenXR/Source/OpenXRHMD/Private/OpenXRHMD.h b/Engine/Plugins/Runtime/OpenXR/Source/OpenXRHMD/Private/OpenXRHMD.h
--- a/Engine/Plugins/Runtime/OpenXR/Source/OpenXRHMD/Private/OpenXRHMD.h
+++ b/Engine/Plugins/Runtime/OpenXR/Source/OpenXRHMD/Private/OpenXRHMD.h
@@ -484,14 +484,12 @@ public:
 	IOpenXRExtensionPluginDelegates& GetIOpenXRExtensionPluginDelegates() override { return *this; }
 	TArray<IOpenXRExtensionPlugin*>& GetExtensionPlugins() override { return ExtensionPlugins; }
 	
+	void SetEnvironmentBlendMode(XrEnvironmentBlendMode NewBlendMode) override;
+	XrEnvironmentBlendMode GetEnvironmentBlendMode() override { return SelectedEnvironmentBlendMode; }
+	TArray<XrEnvironmentBlendMode> GetSupportedEnvironmentBlendModes() const override;
 
-	/* Note: this function will be added to IOpenXRHMD.h in 5.7*/
-	OPENXRHMD_API void SetEnvironmentBlendMode(XrEnvironmentBlendMode NewBlendMode);
-
-	/** Provides the full asymmetric FOV information for the view, which allows proper centering of screen-based techniques based on forward eye gaze 
-	* Note: this function will only be exposed here in 5.6.  In 5.7 it will be added to IOpenXRHMD.h.
-	*/
-	OPENXRHMD_API void GetFullFOVInformation(TArray<XrFovf>& FovInfos) const;
+	/** Provides the full asymmetric FOV information for the view, which allows proper centering of screen-based techniques based on forward eye gaze */
+	void GetFullFOVInformation(TArray<XrFovf>& FovInfos) const override;
 
 	/** Returns shader platform the plugin is currently configured for, in the editor it can change due to preview platforms. */
 	EShaderPlatform GetConfiguredShaderPlatform() const { check(ConfiguredShaderPlatform != EShaderPlatform::SP_NumPlatforms); return ConfiguredShaderPlatform; }
@@ -500,7 +498,6 @@ public:
 	bool RuntimeRequiresRHIContext() const { return bRuntimeRequiresRHIContext; }
 private:
 
-	TArray<XrEnvironmentBlendMode> RetrieveEnvironmentBlendModes() const;
 	FDefaultStereoLayers_LayerRenderParams CalculateEmulatedLayerRenderParams(const FSceneView& InView);
 	class FEmulatedLayersPass* SetupEmulatedLayersRenderPass(FRDGBuilder& GraphBuilder, const FSceneView& InView, TArray<FDefaultStereoLayers::FStereoLayerToRender>& Layers, FRDGTextureRef RenderTarget, FDefaultStereoLayers_LayerRenderParams& OutRenderParams);
 	bool IsEmulatingStereoLayers();
@@ -592,7 +589,6 @@ private:
 	FVector					BasePosition;
 
 	bool					bLayerSupportOpenXRCompliant;
-	bool					bOpenXRInvertAlphaCvarCachedValue;
 	bool					bOpenXRForceStereoLayersEmulationCVarCachedValue;
 	TArray<uint32>			VisibleLayerIds;
 	UE_DEPRECATED(5.6, "This will no longer be needed once OnSetupLayers_RenderThread is removed")
diff --git a/Engine/Plugins/Runtime/OpenXR/Source/OpenXRHMD/Private/OpenXRHMDModule.cpp b/Engine/Plugins/Runtime/OpenXR/Source/OpenXRHMD/Private/OpenXRHMDModule.cpp
--- a/Engine/Plugins/Runtime/OpenXR/Source/OpenXRHMD/Private/OpenXRHMDModule.cpp
+++ b/Engine/Plugins/Runtime/OpenXR/Source/OpenXRHMD/Private/OpenXRHMDModule.cpp
@@ -2,6 +2,7 @@
 
 #include "OpenXRHMDModule.h"
 #include "Misc/EngineVersion.h"
+#include "Misc/ConfigUtilities.h"
 #include "OpenXRHMD.h"
 #include "OpenXRHMD_RenderBridge.h"
 #include "OpenXRCore.h"
@@ -136,6 +137,14 @@ bool FOpenXRHMDModule::PreInit()
 #endif // PLATFORM_WINDOWS
 }
 
+void FOpenXRHMDModule::StartupModule()
+{
+	IOpenXRHMDModule::StartupModule();
+
+	// Pull in CVar settings from OpenXRHMDSettings in BaseEngine.ini
+	UE::ConfigUtilities::ApplyCVarSettingsFromIni(TEXT("/Script/OpenXRHMD.OpenXRHMDSettings"), *GEngineIni, ECVF_SetByProjectSetting);
+}
+
 void FOpenXRHMDModule::ShutdownModule()
 {
 	if (Instance)
@@ -148,6 +157,8 @@ void FOpenXRHMDModule::ShutdownModule()
 		FPlatformProcess::FreeDllHandle(LoaderHandle);
 		LoaderHandle = nullptr;
 	}
+
+	IOpenXRHMDModule::ShutdownModule();
 }
 
 uint64 FOpenXRHMDModule::GetGraphicsAdapterLuid()
diff --git a/Engine/Plugins/Runtime/OpenXR/Source/OpenXRHMD/Private/OpenXRHMDModule.h b/Engine/Plugins/Runtime/OpenXR/Source/OpenXRHMD/Private/OpenXRHMDModule.h
--- a/Engine/Plugins/Runtime/OpenXR/Source/OpenXRHMD/Private/OpenXRHMDModule.h
+++ b/Engine/Plugins/Runtime/OpenXR/Source/OpenXRHMD/Private/OpenXRHMDModule.h
@@ -33,7 +33,8 @@ public:
 		AliasesOut.Add(TEXT("OpenXR"));
 	}
 
-	virtual bool PreInit() override; 
+	virtual bool PreInit() override;
+	virtual void StartupModule() override;
 	virtual void ShutdownModule() override;
 
 	virtual bool IsHMDConnected() override { return true; }
diff --git a/Engine/Plugins/Runtime/OpenXR/Source/OpenXRHMD/Private/OpenXRHMDSettings.cpp b/Engine/Plugins/Runtime/OpenXR/Source/OpenXRHMD/Private/OpenXRHMDSettings.cpp
deleted file mode 100644
diff --git a/Engine/Plugins/Runtime/OpenXR/Source/OpenXRHMD/Public/IOpenXRHMD.h b/Engine/Plugins/Runtime/OpenXR/Source/OpenXRHMD/Public/IOpenXRHMD.h
--- a/Engine/Plugins/Runtime/OpenXR/Source/OpenXRHMD/Public/IOpenXRHMD.h
+++ b/Engine/Plugins/Runtime/OpenXR/Source/OpenXRHMD/Public/IOpenXRHMD.h
@@ -40,5 +40,11 @@ public:
 	OPENXRHMD_API virtual bool GetPoseForTime(int32 DeviceId, FTimespan Timespan, bool& OutTimeWasUsed, FQuat& CurrentOrientation, FVector& CurrentPosition, bool& bProvidedLinearVelocity, FVector& LinearVelocity, bool& bProvidedAngularVelocity, FVector& AngularVelocityAsAxisAndLength, bool& bProvidedLinearAcceleration, FVector& LinearAcceleration, float WorldToMetersScale) = 0;
 
 	OPENXRHMD_API virtual TArray<class IOpenXRExtensionPlugin*>& GetExtensionPlugins() = 0;
+
+	OPENXRHMD_API virtual void SetEnvironmentBlendMode(XrEnvironmentBlendMode NewBlendMode) = 0;
+	OPENXRHMD_API virtual XrEnvironmentBlendMode GetEnvironmentBlendMode() = 0;
+	OPENXRHMD_API virtual TArray<XrEnvironmentBlendMode> GetSupportedEnvironmentBlendModes() const = 0;
+
+	OPENXRHMD_API virtual void GetFullFOVInformation(TArray<XrFovf>& FovInfos) const = 0;
 };
 
diff --git a/Engine/Plugins/Runtime/OpenXR/Source/OpenXRHMD/Public/OpenXRBlueprintFunctionLibrary.h b/Engine/Plugins/Runtime/OpenXR/Source/OpenXRHMD/Public/OpenXRBlueprintFunctionLibrary.h
--- a/Engine/Plugins/Runtime/OpenXR/Source/OpenXRHMD/Public/OpenXRBlueprintFunctionLibrary.h
+++ b/Engine/Plugins/Runtime/OpenXR/Source/OpenXRHMD/Public/OpenXRBlueprintFunctionLibrary.h
@@ -3,22 +3,52 @@
 #pragma once
 
 #include "Kismet/BlueprintFunctionLibrary.h"
+#include "OpenXRCore.h"
 
 #include "OpenXRBlueprintFunctionLibrary.generated.h"
 
 class FOpenXRHMD;
 
+UENUM(BlueprintType)
+enum class EOpenXREnvironmentBlendMode : uint8
+{
+	None = 0,
+	Opaque = XR_ENVIRONMENT_BLEND_MODE_OPAQUE,
+	Additive = XR_ENVIRONMENT_BLEND_MODE_ADDITIVE,
+    AlphaBlend = XR_ENVIRONMENT_BLEND_MODE_ALPHA_BLEND,
+};
+
 UCLASS()
 class OPENXRHMD_API UOpenXRBlueprintFunctionLibrary : public UBlueprintFunctionLibrary
 {
 	GENERATED_UCLASS_BODY()
+
 	/**
-	 * Allows to change the environment blend mode at runtime.
+	 * Sets the OpenXR environment blend mode.
 	 * 
 	 * @param NewBlendMode		The new blend mode to be set, if supported.
 	 */
-	UFUNCTION(BlueprintCallable, Category="OpenXR|Experimental")
-	static void SetEnvironmentBlendMode(int32 NewBlendMode);
+	UFUNCTION(BlueprintCallable, Category="OpenXR|Passthrough")
+	static void SetEnvironmentBlendMode(EOpenXREnvironmentBlendMode NewBlendMode);
+
+	/**
+	 * Gets the OpenXR environment blend mode that is currently in use.
+	 */
+	UFUNCTION(BlueprintCallable, Category="OpenXR|Passthrough")
+	static EOpenXREnvironmentBlendMode GetEnvironmentBlendMode();
+
+	/**
+	 * Gets the OpenXR runtime's supported environment blend modes, sorted in order from highest to lowest runtime preference.
+	 */
+	UFUNCTION(BlueprintCallable, Category="OpenXR|Passthrough")
+	static TArray<EOpenXREnvironmentBlendMode> GetSupportedEnvironmentBlendModes();
+
+	/**
+	 * Returns true if OpenXR composition layer inverted alpha is enabled.
+	 */
+	UFUNCTION(BlueprintCallable, Category="OpenXR|Passthrough")
+	static bool IsCompositionLayerInvertedAlphaEnabled();
+
 
 private:
 	static FOpenXRHMD* GetOpenXRHMD();
diff --git a/Engine/Plugins/Runtime/OpenXR/Source/OpenXRHMD/Public/OpenXRHMDSettings.h b/Engine/Plugins/Runtime/OpenXR/Source/OpenXRHMD/Public/OpenXRHMDSettings.h
--- a/Engine/Plugins/Runtime/OpenXR/Source/OpenXRHMD/Public/OpenXRHMDSettings.h
+++ b/Engine/Plugins/Runtime/OpenXR/Source/OpenXRHMD/Public/OpenXRHMDSettings.h
@@ -3,17 +3,16 @@
 #pragma once
 
 #include "CoreMinimal.h"
-#include "UObject/Object.h"
-#include "Engine/RendererSettings.h"
+#include "Engine/DeveloperSettingsBackedByCVars.h"
 #include "OpenXRHMDSettings.generated.h"
 
 /**
-* Implements the settings for the OpenXR HMD plugin.
+* Implements the settings for the OpenXR plugin.
 */
-UCLASS(config = Game, defaultconfig)
-class OPENXRHMD_API UOpenXRHMDSettings : public UObject
+UCLASS(config=Engine, defaultconfig, meta = (DisplayName = "OpenXR Settings"))
+class OPENXRHMD_API UOpenXRHMDSettings : public UDeveloperSettingsBackedByCVars
 {
-	GENERATED_UCLASS_BODY()
+	GENERATED_BODY()
 
 public:
 	/** Enables foveation provided by the XR_FB_foveation OpenXR extension. */
@@ -22,6 +21,13 @@ public:
 		DisplayName = "Enable XR_FB_foveation extension"))
 	bool bIsFBFoveationEnabled = false;
 
+	/** Enables alpha inversion of the background layer. */
+	UPROPERTY(config, EditAnywhere, Category = "Passthrough", meta = (
+		ConsoleVariable = "xr.OpenXRInvertAlpha",
+		ToolTip = "Enables alpha inversion of the background layer if the XR_EXT_composition_layer_inverted_alpha extension or XR_FB_composition_layer_alpha_blend is supported.", 
+		DisplayName = "Invert scene alpha for passthrough"))
+	bool bOpenXRInvertAlpha = false;
+
 	/** Enable support for OpenXR 1.0. */
 	UPROPERTY(config, EditAnywhere, Category = "OpenXR Versions", meta = (
 		ToolTip = "Enable support for OpenXR 1.0. If multiple versions are supported by the current OpenXR Runtime the latest version will be used.",
@@ -33,4 +39,7 @@ public:
 		ToolTip = "Enable support for OpenXR 1.1. If multiple versions are supported by the current OpenXR Runtime the latest version will be used.",
 		DisplayName = "Enable OpenXR 1.1"))
 	bool bIsOpenXR1_1Enabled = true;
+
+public:
+	virtual FName GetCategoryName() const override { return FName(TEXT("Plugins")); }
 };
diff --git a/Engine/Source/Runtime/Engine/Classes/Engine/RendererSettings.h b/Engine/Source/Runtime/Engine/Classes/Engine/RendererSettings.h
--- a/Engine/Source/Runtime/Engine/Classes/Engine/RendererSettings.h
+++ b/Engine/Source/Runtime/Engine/Classes/Engine/RendererSettings.h
@@ -1418,6 +1418,12 @@ class URendererSettings : public UDeveloperSettings
 		ConfigRestartRequired = true))
 		uint32 bMobileDBuffer : 1;
 
+	UPROPERTY(config, EditAnywhere, Category = Mobile, meta = (
+		ConsoleVariable = "r.Mobile.PropagateAlpha", DisplayName = "Mobile Alpha Output",
+		ToolTip = "Enable r.Mobile.PropagateAlpha to enforce alpha in scene color (overriding r.SceneColorFormat if necessary) and propagate it through the mobile renderer's post-processing chain.",
+		ConfigRestartRequired = true))
+	uint32 bMobileEnableAlphaChannelInPostProcessing : 1;
+
 	UPROPERTY(config, EditAnywhere, Category = Skinning, meta = (
 		ConsoleVariable = "r.GPUSkin.UnlimitedBoneInfluences", DisplayName = "Use Unlimited Bone Influences",
 		ToolTip = "If enabled, a new mesh imported will use unlimited bone buffer instead of fixed MaxBoneInfluences for rendering.",
