Add native support for XR_FB_space_warp and XR_EXT_frame_synthesis to OpenXRHMD, and add a stereo motion vector rendering pass to the mobile renderer which can be used in place of normal velocity rendering, providing motion vector and motion vector depth swapchains at the size requested by space warp/frame synthesis.

OpenXR motion vector rendering is currently mutually exclusive with normal velocity rendering, and is supported only on the Vulkan mobile forward renderer. It can be toggled on with r.Velocity.DirectlyRenderOpenXRMotionVectors (read-only). Space warp/frame synthesis itself is controlled with xr.OpenXRFrameSynthesis (toggleable at runtime, but requires r.Velocity.DirectlyRenderOpenXRMotionVectors=1).

Generating motion vectors and depth at a different resolution from the normal scene textures requires including all primitives in the motion vectors velocity pass, even stationary primitives that would normally be excluded. This is necessary for two reasons: because we need the generated depth to be comprehensive, and because we don't have a pre-existing depth at a matching resolution that we can use with ClipToPrevClip to calculate velocity for all stationary primitives in a single pass.

This change adds two new RenderUtil functions to control the above behavior: VelocityIncludeStationaryPrimitives() and PlatformSupportsOpenXRMotionVectors(). Currently these will both return true when r.Velocity.DirectlyRenderOpenXRMotionVectors=1 and we're on the Vulkan mobile forward renderer. It also adds new virtual functions to IStereoRenderTargetManager to allow requesting motion vector and motion vector depth targets, and slightly restructures swapchain allocation in OpenXRHMD to reduce code duplication.

diff --git a/Engine/Plugins/Runtime/OpenXR/Source/OpenXRHMD/Private/OpenXRHMD.cpp b/Engine/Plugins/Runtime/OpenXR/Source/OpenXRHMD/Private/OpenXRHMD.cpp
--- a/Engine/Plugins/Runtime/OpenXR/Source/OpenXRHMD/Private/OpenXRHMD.cpp
+++ b/Engine/Plugins/Runtime/OpenXR/Source/OpenXRHMD/Private/OpenXRHMD.cpp
@@ -152,8 +152,16 @@ static TAutoConsoleVariable<bool> CVarOpenXRLateUpdateDeviceLocationsAfterReflec
 static TAutoConsoleVariable<bool> CVarOpenXRAlphaInvertPass(
 	TEXT("OpenXR.AlphaInvertPass"),
 	0,
-	TEXT("Whether to run a render pass to un-invert the alpha value from unreal standard to the much more common standard where alpha 0 is fully transparent and alpha 1 is fully opaque."
-	TEXT("This cvar specifically enables the pass for the main XR view.  There is a more general r.ALphaInvertPass which enables it for all renders.")),
+	TEXT("Whether to run a render pass to un-invert the alpha value from unreal standard to the much more common standard where alpha 0 is fully transparent and alpha 1 is fully opaque.")
+	TEXT("This cvar specifically enables the pass for the main XR view.  There is a more general r.ALphaInvertPass which enables it for all renders."),
+	ECVF_RenderThreadSafe);
+
+static TAutoConsoleVariable<bool> CVarOpenXRFrameSynthesis(
+	TEXT("xr.OpenXRFrameSynthesis"),
+	1,
+	TEXT("If true and supported via XR_EXT_frame_synthesis or XR_FB_space_warp, enable and submit motion vector and motion vector depth swapchains for frame synthesis.\n")
+	TEXT("Currently only supported when using the Vulkan mobile renderer, using mobile multi-view, and r.Velocity.DirectlyRenderOpenXRMotionVectors=True.\n")
+	TEXT("Because normal velocity rendering is disabled when r.Velocity.DirectlyRenderOpenXRMotionVectors=True, temporal anti-aliasing and motion blur will be automatically disabled."),
 	ECVF_RenderThreadSafe);
 
 namespace {
@@ -1208,6 +1216,23 @@ void FOpenXRHMD::SetFinalViewRect(FRHICommandListImmediate& RHICmdList, const in
 	XrSwapchainSubImage& EmulationImage = PipelinedLayerStateRendering.EmulatedLayerState.EmulationImages[ViewIndex];
 	EmulationImage.imageArrayIndex = ColorImage.imageArrayIndex;
 	EmulationImage.imageRect = ColorImage.imageRect;
+
+	FIntPoint MotionVectorSize;
+	if (GetRecommendedMotionVectorTextureSize(MotionVectorSize))
+	{
+		const FIntRect MotionVectorRect(FIntPoint(0, 0), MotionVectorSize);
+
+		XrSwapchainSubImage& MotionVectorImage = PipelinedLayerStateRendering.MotionVectorImages[ViewIndex];
+		MotionVectorImage.imageArrayIndex = ColorImage.imageArrayIndex;
+		MotionVectorImage.imageRect = {
+			{ MotionVectorRect.Min.X, MotionVectorRect.Min.Y },
+			{ MotionVectorRect.Width(), MotionVectorRect.Height() }
+		};
+
+		XrSwapchainSubImage& MotionVectorDepthImage = PipelinedLayerStateRendering.MotionVectorDepthImages[ViewIndex];
+		MotionVectorDepthImage.imageArrayIndex = ColorImage.imageArrayIndex;
+		MotionVectorDepthImage.imageRect = MotionVectorImage.imageRect;
+	}
 }
 
 EStereoscopicPass FOpenXRHMD::GetViewPassForIndex(bool bStereoRequested, int32 ViewIndex) const
@@ -1342,6 +1367,24 @@ void FOpenXRHMD::SetupViewFamily(FSceneViewFamily& InViewFamily)
 		InViewFamily.EngineShowFlags.Vignette = 0;
 		InViewFamily.EngineShowFlags.Bloom = 0;
 	}
+
+	if (InViewFamily.Scene)
+	{
+		EShaderPlatform ShaderPlatform = InViewFamily.Scene->GetShaderPlatform();
+		if (PlatformSupportsOpenXRMotionVectors(ShaderPlatform) && bIsMobileMultiViewEnabled)
+		{
+			static FShaderPlatformCachedIniValue<bool> FrameSynthesisIniValue(TEXT("xr.OpenXRFrameSynthesis"));
+			const bool bFrameSynthesisCVarEnabled = FrameSynthesisIniValue.Get(ShaderPlatform);
+			const bool bFrameSynthesisExtensionSupported =
+				IsExtensionEnabled(XR_EXT_FRAME_SYNTHESIS_EXTENSION_NAME) ||
+				IsExtensionEnabled(XR_FB_SPACE_WARP_EXTENSION_NAME);
+
+			if (bFrameSynthesisCVarEnabled && bFrameSynthesisExtensionSupported)
+			{
+				InViewFamily.EngineShowFlags.StereoMotionVectors = true;
+			}
+		}
+	}
 }
 
 void FOpenXRHMD::SetupView(FSceneViewFamily& InViewFamily, FSceneView& InView)
@@ -1357,10 +1400,16 @@ void FOpenXRHMD::BeginRenderViewFamily(FSceneViewFamily& InViewFamily)
 	PipelinedLayerStateRendering.DepthLayers.SetNum(ViewConfigCount);
 	PipelinedLayerStateRendering.EmulatedLayerState.CompositedProjectionLayers.SetNum(ViewConfigCount);
 
+	PipelinedLayerStateRendering.FrameSynthesisLayers.SetNum(ViewConfigCount);
+	PipelinedLayerStateRendering.SpaceWarpLayers.SetNum(ViewConfigCount);
+
 	PipelinedLayerStateRendering.ColorImages.SetNum(PipelinedFrameStateRendering.ViewConfigs.Num());
 	PipelinedLayerStateRendering.DepthImages.SetNum(PipelinedFrameStateRendering.ViewConfigs.Num());
 	PipelinedLayerStateRendering.EmulatedLayerState.EmulationImages.SetNum(PipelinedFrameStateRendering.ViewConfigs.Num());
 
+	PipelinedLayerStateRendering.MotionVectorImages.SetNum(PipelinedFrameStateRendering.ViewConfigs.Num());
+	PipelinedLayerStateRendering.MotionVectorDepthImages.SetNum(PipelinedFrameStateRendering.ViewConfigs.Num());
+
 	if (bCompositionLayerColorScaleBiasSupported)
 	{
 		PipelinedLayerStateRendering.LayerColorScaleAndBias = { LayerColorScale, LayerColorBias };
@@ -1411,6 +1460,8 @@ void FOpenXRHMD::PostRenderViewFamily_RenderThread(FRDGBuilder& GraphBuilder, FS
 		XrSwapchainSubImage& ColorImage = PipelinedLayerStateRendering.ColorImages[ViewIndex];
 		XrSwapchainSubImage& DepthImage = PipelinedLayerStateRendering.DepthImages[ViewIndex];
 		XrSwapchainSubImage& EmulationImage = PipelinedLayerStateRendering.EmulatedLayerState.EmulationImages[ViewIndex];
+		XrSwapchainSubImage& MotionVectorImage = PipelinedLayerStateRendering.MotionVectorImages[ViewIndex];
+		XrSwapchainSubImage& MotionVectorDepthImage = PipelinedLayerStateRendering.MotionVectorDepthImages[ViewIndex];
 
 		ColorImage.swapchain = PipelinedLayerStateRendering.ColorSwapchain.IsValid() ? static_cast<FOpenXRSwapchain*>(PipelinedLayerStateRendering.ColorSwapchain.Get())->GetHandle() : XR_NULL_HANDLE;
 		if (EnumHasAnyFlags(PipelinedLayerStateRendering.LayerStateFlags, EOpenXRLayerStateFlags::SubmitDepthLayer))
@@ -1421,6 +1472,13 @@ void FOpenXRHMD::PostRenderViewFamily_RenderThread(FRDGBuilder& GraphBuilder, FS
 		{
 			EmulationImage.swapchain = PipelinedLayerStateRendering.EmulatedLayerState.EmulationSwapchain.IsValid() ? static_cast<FOpenXRSwapchain*>(PipelinedLayerStateRendering.EmulatedLayerState.EmulationSwapchain.Get())->GetHandle() : XR_NULL_HANDLE;
 		}
+		if (EnumHasAnyFlags(PipelinedLayerStateRendering.LayerStateFlags, EOpenXRLayerStateFlags::SubmitMotionVectorLayer))
+		{
+			check(PipelinedLayerStateRendering.MotionVectorSwapchain && PipelinedLayerStateRendering.MotionVectorDepthSwapchain);
+
+			MotionVectorImage.swapchain = PipelinedLayerStateRendering.MotionVectorSwapchain.IsValid() ? static_cast<FOpenXRSwapchain*>(PipelinedLayerStateRendering.MotionVectorSwapchain.Get())->GetHandle() : XR_NULL_HANDLE;
+			MotionVectorDepthImage.swapchain = PipelinedLayerStateRendering.MotionVectorDepthSwapchain.IsValid() ? static_cast<FOpenXRSwapchain*>(PipelinedLayerStateRendering.MotionVectorDepthSwapchain.Get())->GetHandle() : XR_NULL_HANDLE;
+		}
 
 		XrCompositionLayerProjectionView& Projection = PipelinedLayerStateRendering.ProjectionLayers[ViewIndex];
 
@@ -1455,6 +1513,61 @@ void FOpenXRHMD::PostRenderViewFamily_RenderThread(FRDGBuilder& GraphBuilder, FS
 			CompositedProjection.next = nullptr;
 			CompositedProjection.subImage = EmulationImage;
 		}
+		if (EnumHasAnyFlags(PipelinedLayerStateRendering.LayerStateFlags, EOpenXRLayerStateFlags::SubmitMotionVectorLayer))
+		{
+			if (IsExtensionEnabled(XR_EXT_FRAME_SYNTHESIS_EXTENSION_NAME))
+			{
+				XrFrameSynthesisInfoEXT& FrameSynthesisLayerInfo = PipelinedLayerStateRendering.FrameSynthesisLayers[ViewIndex];
+
+				FTransform TrackingToWorld = GetTrackingToWorldTransform();
+				FTransform TrackingSpaceDeltaPose = TrackingToWorld * FrameSynthesisLastTrackingToWorld.Inverse();
+				FrameSynthesisLastTrackingToWorld = TrackingToWorld;
+				FTransform BaseTransform = FTransform(GetBaseOrientation(), GetBasePosition());
+				TrackingSpaceDeltaPose = BaseTransform.Inverse() * TrackingSpaceDeltaPose * BaseTransform;
+
+				FrameSynthesisLayerInfo.type = XR_TYPE_FRAME_SYNTHESIS_INFO_EXT;
+				FrameSynthesisLayerInfo.next = nullptr;
+				FrameSynthesisLayerInfo.layerFlags = 0;
+				FrameSynthesisLayerInfo.motionVectorSubImage = MotionVectorImage;
+				FrameSynthesisLayerInfo.motionVectorScale = { 1.0f, 1.0f, 1.0f, 0.0f };
+				FrameSynthesisLayerInfo.motionVectorOffset = {};
+				FrameSynthesisLayerInfo.appSpaceDeltaPose = ToXrPose(TrackingSpaceDeltaPose);
+				FrameSynthesisLayerInfo.depthSubImage = MotionVectorDepthImage;
+				FrameSynthesisLayerInfo.minDepth = 0.0f;
+				FrameSynthesisLayerInfo.maxDepth = 1.0f;
+				FrameSynthesisLayerInfo.nearZ = FLT_MAX;
+				FrameSynthesisLayerInfo.farZ = NearZ;
+
+				Projection.next = &FrameSynthesisLayerInfo;
+			}
+			else if (IsExtensionEnabled(XR_FB_SPACE_WARP_EXTENSION_NAME))
+			{
+				XrCompositionLayerSpaceWarpInfoFB& SpaceWarpLayerInfo = PipelinedLayerStateRendering.SpaceWarpLayers[ViewIndex];
+
+				FTransform TrackingToWorld = GetTrackingToWorldTransform();
+				FTransform TrackingSpaceDeltaPose = TrackingToWorld * FrameSynthesisLastTrackingToWorld.Inverse();
+				FrameSynthesisLastTrackingToWorld = TrackingToWorld;
+				FTransform BaseTransform = FTransform(GetBaseOrientation(), GetBasePosition());
+				TrackingSpaceDeltaPose = BaseTransform.Inverse() * TrackingSpaceDeltaPose * BaseTransform;
+
+				SpaceWarpLayerInfo.type = XR_TYPE_COMPOSITION_LAYER_SPACE_WARP_INFO_FB;
+				SpaceWarpLayerInfo.next = nullptr;
+				SpaceWarpLayerInfo.layerFlags = 0;
+				SpaceWarpLayerInfo.motionVectorSubImage = MotionVectorImage;
+				SpaceWarpLayerInfo.appSpaceDeltaPose = ToXrPose(TrackingSpaceDeltaPose);
+				SpaceWarpLayerInfo.depthSubImage = MotionVectorDepthImage;
+				SpaceWarpLayerInfo.minDepth = 0.0f;
+				SpaceWarpLayerInfo.maxDepth = 1.0f;
+				SpaceWarpLayerInfo.nearZ = FLT_MAX;
+				SpaceWarpLayerInfo.farZ = NearZ;
+
+				Projection.next = &SpaceWarpLayerInfo;
+			}
+			else
+			{
+				checkNoEntry();
+			}
+		}
 
 		for (IOpenXRExtensionPlugin* Module : ExtensionPlugins)
 		{
@@ -1549,6 +1662,7 @@ FOpenXRHMD::FOpenXRHMD(const FAutoRegister& AutoRegister, XrInstance InInstance,
 	, LayerColorScale{ 1.0f, 1.0f, 1.0f, 1.0f }
 	, LayerColorBias{ 0.0f, 0.0f, 0.0f, 0.0f }
 	, bxrGetSystemPropertiesSuccessful(false)
+	, RecommendedMotionVectorTextureSize(0, 0)
 {
 	check(OpenXRAPIVersion != EOpenXRAPIVersion::V_INVALID);
 
@@ -1604,6 +1718,8 @@ FOpenXRHMD::FOpenXRHMD(const FAutoRegister& AutoRegister, XrInstance InInstance,
 		Module->BindExtensionPluginDelegates(*this);
 		bRuntimeRequiresRHIContext |= Module->RequiresRHIContext();
 	}
+
+	FrameSynthesisLastTrackingToWorld = FTransform::Identity;
 }
 
 PRAGMA_DISABLE_DEPRECATION_WARNINGS
@@ -1757,8 +1873,11 @@ void FOpenXRHMD::EnumerateViews(FPipelinedFrameState& PipelineState)
 	// Enumerate the viewport configuration views
 	uint32 ViewConfigCount = 0;
 	TArray<XrViewConfigurationViewFovEPIC> ViewFov;
+	TArray<XrFrameSynthesisConfigViewEXT> ViewFrameSynthesis;
+
 	XR_ENSURE(xrEnumerateViewConfigurationViews(Instance, System, SelectedViewConfigurationType, 0, &ViewConfigCount, nullptr));
 	ViewFov.SetNum(ViewConfigCount);
+	ViewFrameSynthesis.SetNum(ViewConfigCount);
 	PipelineState.ViewConfigs.Empty(ViewConfigCount);
 	for (uint32 ViewIndex = 0; ViewIndex < ViewConfigCount; ViewIndex++)
 	{
@@ -1769,11 +1888,22 @@ void FOpenXRHMD::EnumerateViews(FPipelinedFrameState& PipelineState)
 		ViewFov[ViewIndex].next = nullptr;
 		View.next = bViewConfigurationFovSupported ? &ViewFov[ViewIndex] : nullptr;
 
+		if (IsExtensionEnabled(XR_EXT_FRAME_SYNTHESIS_EXTENSION_NAME))
+		{
+			check(!IsExtensionEnabled(XR_FB_SPACE_WARP_EXTENSION_NAME));
+
+			ViewFrameSynthesis[ViewIndex].type = XR_TYPE_FRAME_SYNTHESIS_CONFIG_VIEW_EXT;
+			ViewFrameSynthesis[ViewIndex].next = View.next;
+			View.next = &ViewFrameSynthesis[ViewIndex];
+		}
+
 		for (IOpenXRExtensionPlugin* Module : ExtensionPlugins)
 		{
 			View.next = Module->OnEnumerateViewConfigurationViews(Instance, System, SelectedViewConfigurationType, ViewIndex, View.next);
 		}
 
+		XrFrameSynthesisConfigViewEXT FrameSynthesisViewConfig = { XR_TYPE_FRAME_SYNTHESIS_CONFIG_VIEW_EXT };
+		
 		PipelineState.ViewConfigs.Add(View);
 	}
 	XR_ENSURE(xrEnumerateViewConfigurationViews(Instance, System, SelectedViewConfigurationType, ViewConfigCount, &ViewConfigCount, PipelineState.ViewConfigs.GetData()));
@@ -1809,6 +1939,14 @@ void FOpenXRHMD::EnumerateViews(FPipelinedFrameState& PipelineState)
 			View.pose = ToXrPose(FTransform::Identity);
 		}
 	}
+
+	if (IsExtensionEnabled(XR_EXT_FRAME_SYNTHESIS_EXTENSION_NAME) && PipelineState.ViewConfigs.Num())
+	{
+		// Variable size between views is not currently supported
+		XrFrameSynthesisConfigViewEXT& FrameSynthConfig = ViewFrameSynthesis[0];
+		RecommendedMotionVectorTextureSize = FIntPoint(FrameSynthConfig.recommendedMotionVectorImageRectWidth, FrameSynthConfig.recommendedMotionVectorImageRectHeight);
+		check(RecommendedMotionVectorTextureSize.X && RecommendedMotionVectorTextureSize.Y);
+	}
 }
 
 void FOpenXRHMD::BuildOcclusionMeshes()
@@ -1980,6 +2118,12 @@ bool FOpenXRHMD::AcquireSystemIdAndProperties()
 	{
 		OpenXR::AppendChainStruct(Tail, &UserPresenceSystemProperties);
 	}
+	XrSystemSpaceWarpPropertiesFB SpaceWarpViewConfig = { XR_TYPE_SYSTEM_SPACE_WARP_PROPERTIES_FB };
+	if (IsExtensionEnabled(XR_FB_SPACE_WARP_EXTENSION_NAME))
+	{
+		check(!IsExtensionEnabled(XR_EXT_FRAME_SYNTHESIS_EXTENSION_NAME));
+		OpenXR::AppendChainStruct(Tail, &SpaceWarpViewConfig);
+	}
 
 	XrResult GetSystemPropsResult = xrGetSystemProperties(Instance, System, &SystemProperties);
 	XR_ENSURE(GetSystemPropsResult);
@@ -1991,6 +2135,12 @@ bool FOpenXRHMD::AcquireSystemIdAndProperties()
 	bSupportsHandTracking = HandTrackingSystemProperties.supportsHandTracking == XR_TRUE;
 	bUserPresenceSupported = UserPresenceSystemProperties.supportsUserPresence == XR_TRUE;
 
+	if (IsExtensionEnabled(XR_FB_SPACE_WARP_EXTENSION_NAME))
+	{
+		RecommendedMotionVectorTextureSize = FIntPoint(SpaceWarpViewConfig.recommendedMotionVectorImageRectWidth, SpaceWarpViewConfig.recommendedMotionVectorImageRectHeight);
+		check(RecommendedMotionVectorTextureSize.X && RecommendedMotionVectorTextureSize.Y);
+	}
+
 	// We have built a linked list to function-scope structs and attached it to this member variable, those pointers will be bad after this function exits.
 	// Currently we never want to use them again, so we can just null out the next pointer.
 	SystemProperties.next = nullptr;
@@ -2423,6 +2573,62 @@ XrSpace FOpenXRHMD::GetTrackingSpace() const
 	}
 }
 
+bool FOpenXRHMD::AllocateSwapchainTextures_RenderThread(const FOpenXRSwapchainProperties& InSwapchainProperties, FXRSwapChainPtr& InOutSwapchain, uint8& OutActualFormat)
+{
+	check(IsInRenderingThread());
+
+	FReadScopeLock Lock(SessionHandleMutex);
+	if (!Session)
+	{
+		return false;
+	}
+	
+	const FRHITexture* const SwapchainTexture = InOutSwapchain == nullptr ? nullptr : InOutSwapchain->GetTexture2DArray() ? InOutSwapchain->GetTexture2DArray() : InOutSwapchain->GetTexture2D();
+
+	if (InOutSwapchain == nullptr ||
+		SwapchainTexture == nullptr ||
+		SwapchainTexture->GetFormat() != InSwapchainProperties.Format || 
+		SwapchainTexture->GetSizeX() != InSwapchainProperties.SizeX ||
+		SwapchainTexture->GetSizeY() != InSwapchainProperties.SizeY ||
+		SwapchainTexture->GetDesc().ArraySize != InSwapchainProperties.ArraySize ||
+		SwapchainTexture->GetNumMips() != InSwapchainProperties.NumMips ||
+		SwapchainTexture->GetNumSamples() != InSwapchainProperties.NumSamples ||
+		SwapchainTexture->GetFlags() != InSwapchainProperties.CreateFlags)
+	{
+		InOutSwapchain = RenderBridge->CreateSwapchain(
+			Session,
+			InSwapchainProperties.Format,
+			OutActualFormat,
+			InSwapchainProperties.SizeX,
+			InSwapchainProperties.SizeY,
+			InSwapchainProperties.ArraySize,
+			InSwapchainProperties.NumMips,
+			InSwapchainProperties.NumSamples,
+			InSwapchainProperties.CreateFlags,
+			InSwapchainProperties.ClearValueBinding,
+			InSwapchainProperties.AuxiliaryCreateFlags);
+		
+		if (InOutSwapchain)
+		{
+			UE_LOG(LogHMD, Verbose, TEXT("Allocated %.*s OpenXR Swapchain of Format: %u, SizeX: %u, SizeY: %u, ArraySize: %u, NumMips: %u, NumSamples: %u, Flags: %llu"),
+				InSwapchainProperties.DebugName.Len(),
+				InSwapchainProperties.DebugName.GetData(),
+				InSwapchainProperties.Format,
+				InSwapchainProperties.SizeX,
+				InSwapchainProperties.SizeY,
+				InSwapchainProperties.ArraySize,
+				InSwapchainProperties.NumMips,
+				InSwapchainProperties.NumSamples,
+				InSwapchainProperties.CreateFlags);
+			
+        	InOutSwapchain->SetDebugLabel(InSwapchainProperties.DebugName);
+			return true;
+        }
+	}
+
+	return false;
+}
+
 bool FOpenXRHMD::IsInitialized() const
 {
 	return Instance != XR_NULL_HANDLE;
@@ -2540,6 +2746,25 @@ int32 FOpenXRHMD::AcquireDepthTexture()
 	if (Session)
 	{
 		const FXRSwapChainPtr& DepthSwapchain = PipelinedLayerStateRendering.DepthSwapchain;
+		const FXRSwapChainPtr& MotionVectorSwapchain = PipelinedLayerStateRendering.MotionVectorSwapchain;
+		const FXRSwapChainPtr& MotionVectorDepthSwapchain = PipelinedLayerStateRendering.MotionVectorDepthSwapchain;
+
+		// In the future, this function should be merged with AcquireColorTexture() into a single function that acquires all textures
+		// We currently acquire all textures at the same time in FSceneViewport::EnqueueBeginRenderFrame and only use the index returned by AcquireColorTexture()
+		if (MotionVectorSwapchain && MotionVectorDepthSwapchain)
+		{
+			if (bIsAcquireOnAnyThreadSupported)
+			{
+				MotionVectorSwapchain->IncrementSwapChainIndex_RHIThread();
+				MotionVectorDepthSwapchain->IncrementSwapChainIndex_RHIThread();
+			}
+
+			const int32 MotionVectorIndex = MotionVectorSwapchain->GetSwapChainIndex_RHIThread();
+			const int32 MotionVectorDepthIndex = MotionVectorDepthSwapchain->GetSwapChainIndex_RHIThread();
+			
+			check(MotionVectorIndex == MotionVectorDepthIndex);
+		}
+
 		if (DepthSwapchain)
 		{
 			if (bIsAcquireOnAnyThreadSupported)
@@ -2602,38 +2827,47 @@ bool FOpenXRHMD::AllocateRenderTargetTextures(uint32 SizeX, uint32 SizeY, uint8
 
 	FClearValueBinding ClearColor = FClearValueBinding::Transparent;
 
-	uint8 ActualFormat = Format;
 	FXRSwapChainPtr& Swapchain = PipelinedLayerStateRendering.ColorSwapchain;
-	const FRHITexture* const SwapchainTexture = Swapchain == nullptr ? nullptr : Swapchain->GetTexture2DArray() ? Swapchain->GetTexture2DArray() : Swapchain->GetTexture2D();
-	if (Swapchain == nullptr || SwapchainTexture == nullptr || Format != LastRequestedColorSwapchainFormat || SwapchainTexture->GetSizeX() != SizeX || SwapchainTexture->GetSizeY() != SizeY)
+	uint8 ActualFormat = Format;
+	ensureMsgf(NumSamples == 1, TEXT("OpenXR supports MSAA swapchains, but engine logic expects the swapchain target to be 1x."));
+	
+	const FOpenXRSwapchainProperties ColorSwapchainProperties = {
+		TEXT("ColorSwapchain"),
+		Format,
+		SizeX,
+		SizeY,
+		static_cast<uint32>(bIsMobileMultiViewEnabled ? 2 : 1),
+		NumMips,
+		NumSamples,
+		UnifiedCreateFlags,
+		ClearColor,
+		AuxiliaryCreateFlags
+	};
+	bool bAllocatedColorTextures = AllocateSwapchainTextures_RenderThread(ColorSwapchainProperties, Swapchain, ActualFormat);
+	// Image will be acquired by the renderer if supported, if not we acquire it ahead of time here
+	if (!bIsAcquireOnAnyThreadSupported)
 	{
-		ensureMsgf(NumSamples == 1, TEXT("OpenXR supports MSAA swapchains, but engine logic expects the swapchain target to be 1x."));
-
-		Swapchain = RenderBridge->CreateSwapchain(Session, Format, ActualFormat, SizeX, SizeY, bIsMobileMultiViewEnabled ? 2 : 1, NumMips, NumSamples, UnifiedCreateFlags, ClearColor, AuxiliaryCreateFlags);
-		if (!Swapchain)
+		FRHICommandListImmediate& RHICmdList = FRHICommandListImmediate::Get();
+		RHICmdList.EnqueueLambda([this, Swapchain](FRHICommandListImmediate&)
 		{
-			return false;
-		}
-		Swapchain->SetDebugLabel(TEXT("ColorSwapchain"));
-
-		// Image will be acquired by the viewport if supported, if not we acquire it ahead of time here
-		if (!bIsAcquireOnAnyThreadSupported)
-		{
-			FRHICommandListImmediate& RHICmdList = FRHICommandListImmediate::Get();
-			RHICmdList.EnqueueLambda([this, Swapchain](FRHICommandListImmediate&)
+			RenderBridge->RunOnRHISubmissionThread([Swapchain]()
 			{
-				RenderBridge->RunOnRHISubmissionThread([Swapchain]()
-				{
-					Swapchain->IncrementSwapChainIndex_RHIThread();
-				});
+				Swapchain->IncrementSwapChainIndex_RHIThread();
 			});
-			RHICmdList.ImmediateFlush(EImmediateFlushType::FlushRHIThread);
-		}
+		});
+		RHICmdList.ImmediateFlush(EImmediateFlushType::FlushRHIThread);
+	}
+	if (Swapchain)
+	{
 		if (FBFoveationImageGenerator && FBFoveationImageGenerator->IsFoveationExtensionEnabled())
 		{
-			FBFoveationImageGenerator->UpdateFoveationImages(/* bReallocatedSwapchain */ true);
+			FBFoveationImageGenerator->UpdateFoveationImages(/* bReallocatedSwapchain */ bAllocatedColorTextures);
 		}
 	}
+	else
+	{
+		return false;
+	}
 
 	// Grab the presentation texture out of the swapchain.
 	OutTargetableTextures = Swapchain->GetSwapChain();
@@ -2644,29 +2878,89 @@ bool FOpenXRHMD::AllocateRenderTargetTextures(uint32 SizeX, uint32 SizeY, uint8
 	if (IsEmulatingStereoLayers() && (SystemProperties.graphicsProperties.maxLayerCount > 1))
 	{
 		// If we have at least two native layers, use non-background layer to render the composited image of all the emulated face locked layers.
-		FXRSwapChainPtr& EmulationSwapchain = PipelinedLayerStateRendering.EmulatedLayerState.EmulationSwapchain;
-		const FRHITexture* const EmulationSwapchainTexture = EmulationSwapchain == nullptr ? nullptr : EmulationSwapchain->GetTexture2DArray() ? EmulationSwapchain->GetTexture2DArray() : EmulationSwapchain->GetTexture2D();
-		if (EmulationSwapchain == nullptr || EmulationSwapchainTexture == nullptr || EmulationSwapchainTexture->GetSizeX() != SizeX || EmulationSwapchainTexture->GetSizeY() != SizeY)
+		const ETextureCreateFlags EmulationCreateFlags = TexCreate_Dynamic | TexCreate_ShaderResource | TexCreate_RenderTargetable;
+		uint8 UnusedActualFormat = 0;
+
+		const FOpenXRSwapchainProperties EmulationSwapchainProperties = {
+			TEXT("EmulationSwapchain"),
+			IStereoRenderTargetManager::GetStereoLayerPixelFormat(),
+			SizeX,
+			SizeY,
+			static_cast<uint32>(bIsMobileMultiViewEnabled ? 2 : 1),
+			NumMips,
+			NumSamples,
+			EmulationCreateFlags,
+			FClearValueBinding::Transparent
+		};
+
+		AllocateSwapchainTextures_RenderThread(EmulationSwapchainProperties, PipelinedLayerStateRendering.EmulatedLayerState.EmulationSwapchain, UnusedActualFormat);
+
+		// Image will be acquired by the renderer if supported, if not we acquire it ahead of time here
+		if (!bIsAcquireOnAnyThreadSupported)
 		{
-			const ETextureCreateFlags EmulationCreateFlags = TexCreate_Dynamic | TexCreate_ShaderResource | TexCreate_RenderTargetable;
+			FRHICommandListImmediate& RHICmdList = FRHICommandListImmediate::Get();
+			FXRSwapChainPtr& EmulationSwapchain = PipelinedLayerStateRendering.EmulatedLayerState.EmulationSwapchain;
+			RHICmdList.EnqueueLambda([this, EmulationSwapchain](FRHICommandListImmediate&)
+			{
+				RenderBridge->RunOnRHISubmissionThread([EmulationSwapchain]
+				{
+					EmulationSwapchain->IncrementSwapChainIndex_RHIThread();
+				});
+			});
+			RHICmdList.ImmediateFlush(EImmediateFlushType::FlushRHIThread);
+		}
+	}
 
-			uint8 UnusedActualFormat = 0;
-			EmulationSwapchain = RenderBridge->CreateSwapchain(Session, IStereoRenderTargetManager::GetStereoLayerPixelFormat(), UnusedActualFormat, SizeX, SizeY, bIsMobileMultiViewEnabled ? 2 : 1, NumMips, NumSamples, EmulationCreateFlags, FClearValueBinding::Transparent);
-			EmulationSwapchain->SetDebugLabel(TEXT("EmulationSwapchain"));
+	FIntPoint MotionVectorSwapchainSize = FIntPoint();
+	if (GetRecommendedMotionVectorTextureSize(MotionVectorSwapchainSize))
+	{
+		const ETextureCreateFlags MotionVectorCreateFlags = TexCreate_RenderTargetable | TexCreate_ResolveTargetable | TexCreate_ShaderResource | TexCreate_InputAttachmentRead | TexCreate_Dynamic;
+		uint8 UnusedActualFormat = 0;
 
-			// Image will be acquired by SetupFrameLayers_RenderThread if supported, if not we acquire it ahead of time here
-			if (!bIsAcquireOnAnyThreadSupported)
+		const FOpenXRSwapchainProperties MotionVectorSwapchainProperties = {
+			TEXT("MotionVectorSwapchain"),
+			PF_FloatRGBA,
+			static_cast<uint32>(MotionVectorSwapchainSize.X),
+			static_cast<uint32>(MotionVectorSwapchainSize.Y),
+			static_cast<uint32>(bIsMobileMultiViewEnabled ? 2 : 1),
+			1,
+			1,
+			MotionVectorCreateFlags,
+			FClearValueBinding::Transparent
+		};
+
+		AllocateSwapchainTextures_RenderThread(MotionVectorSwapchainProperties, PipelinedLayerStateRendering.MotionVectorSwapchain, UnusedActualFormat);
+
+		const FOpenXRSwapchainProperties MotionVectorDepthSwapchainProperties = {
+			TEXT("MotionVectorDepthSwapchain"),
+			PF_DepthStencil,
+			static_cast<uint32>(MotionVectorSwapchainSize.X),
+			static_cast<uint32>(MotionVectorSwapchainSize.Y),
+			static_cast<uint32>(bIsMobileMultiViewEnabled ? 2 : 1),
+			1,
+			1,
+			TexCreate_DepthStencilTargetable | TexCreate_ShaderResource | TexCreate_InputAttachmentRead | TexCreate_Dynamic,
+			FClearValueBinding::DepthZero,
+			TexCreate_None
+		};
+		
+		AllocateSwapchainTextures_RenderThread(MotionVectorDepthSwapchainProperties, PipelinedLayerStateRendering.MotionVectorDepthSwapchain, UnusedActualFormat);
+
+		// Images will be acquired by the renderer if supported, if not we acquire it ahead of time here
+		if (!bIsAcquireOnAnyThreadSupported)
+		{
+			FRHICommandListImmediate& RHICmdList = FRHICommandListImmediate::Get();
+			FXRSwapChainPtr& MotionVectorSwapchain = PipelinedLayerStateRendering.MotionVectorSwapchain;
+			FXRSwapChainPtr& MotionVectorDepthSwapchain = PipelinedLayerStateRendering.MotionVectorDepthSwapchain;
+			RHICmdList.EnqueueLambda([this, MotionVectorSwapchain, MotionVectorDepthSwapchain](FRHICommandListImmediate&)
 			{
-				FRHICommandListImmediate& RHICmdList = FRHICommandListImmediate::Get();
-				RHICmdList.EnqueueLambda([this, EmulationSwapchain](FRHICommandListImmediate&)
+				RenderBridge->RunOnRHISubmissionThread([MotionVectorSwapchain, MotionVectorDepthSwapchain]
 				{
-					RenderBridge->RunOnRHISubmissionThread([EmulationSwapchain]
-					{
-						EmulationSwapchain->IncrementSwapChainIndex_RHIThread();
-					});
+					MotionVectorSwapchain->IncrementSwapChainIndex_RHIThread();
+					MotionVectorDepthSwapchain->IncrementSwapChainIndex_RHIThread();
 				});
-				RHICmdList.ImmediateFlush(EImmediateFlushType::FlushRHIThread);
-			}
+			});
+			RHICmdList.ImmediateFlush(EImmediateFlushType::FlushRHIThread);
 		}
 	}
 
@@ -2686,46 +2980,48 @@ void FOpenXRHMD::AllocateDepthTextureInternal(uint32 SizeX, uint32 SizeY, uint32
 		return;
 	}
 
-	FXRSwapChainPtr& DepthSwapchain = PipelinedLayerStateRendering.DepthSwapchain;
-	const FRHITexture* const DepthSwapchainTexture = DepthSwapchain == nullptr ? nullptr : DepthSwapchain->GetTexture2DArray() ? DepthSwapchain->GetTexture2DArray() : DepthSwapchain->GetTexture2D();
-	if (DepthSwapchain == nullptr || DepthSwapchainTexture == nullptr ||
-		DepthSwapchainTexture->GetSizeX() != SizeX || DepthSwapchainTexture->GetSizeY() != SizeY || DepthSwapchainTexture->GetDesc().ArraySize != InArraySize)
-	{
-		// We're only creating a 1x target here, but we don't know whether it'll be the targeted texture
-		// or the resolve texture. Because of this, we unify the input flags.
-		ETextureCreateFlags UnifiedCreateFlags = TexCreate_DepthStencilTargetable | TexCreate_ShaderResource | TexCreate_InputAttachmentRead;
-
-		// This is not a static swapchain
-		UnifiedCreateFlags |= TexCreate_Dynamic;
+	// We're only creating a 1x target here, but we don't know whether it'll be the targeted texture
+	// or the resolve texture. Because of this, we unify the input flags.
+	ETextureCreateFlags UnifiedCreateFlags = TexCreate_DepthStencilTargetable | TexCreate_ShaderResource | TexCreate_InputAttachmentRead;
 
-		// We assume this could be used as a resolve target
-		UnifiedCreateFlags |= TexCreate_DepthStencilResolveTarget;
+	// This is not a static swapchain
+	UnifiedCreateFlags |= TexCreate_Dynamic;
 
-		ensureMsgf(NumSamples == 1, TEXT("OpenXR supports MSAA swapchains, but engine logic expects the swapchain target to be 1x."));
-		constexpr uint32 NumSamplesExpected = 1;
-		constexpr uint32 NumMipsExpected = 1;
+	// We assume this could be used as a resolve target
+	UnifiedCreateFlags |= TexCreate_DepthStencilResolveTarget;
+
+	ensureMsgf(NumSamples == 1, TEXT("OpenXR supports MSAA swapchains, but engine logic expects the swapchain target to be 1x."));
+	constexpr uint32 NumSamplesExpected = 1;
+	constexpr uint32 NumMipsExpected = 1;
+	uint8 UnusedActualFormat = 0;
+
+	const FOpenXRSwapchainProperties DepthSwapchainProperties = {
+		TEXT("DepthSwapchain"),
+		PF_DepthStencil,
+		SizeX,
+		SizeY,
+		InArraySize,
+		NumMipsExpected,
+		NumSamplesExpected,
+		UnifiedCreateFlags,
+		FClearValueBinding::DepthFar
+	};
 
-		uint8 UnusedActualFormat = 0;
-		DepthSwapchain = RenderBridge->CreateSwapchain(Session, PF_DepthStencil, UnusedActualFormat, SizeX, SizeY, InArraySize, NumMipsExpected, NumSamplesExpected, UnifiedCreateFlags, FClearValueBinding::DepthFar);
-		if (!DepthSwapchain)
-		{
-			return;
-		}
-		DepthSwapchain->SetDebugLabel(TEXT("DepthSwapchain"));
+	AllocateSwapchainTextures_RenderThread(DepthSwapchainProperties, PipelinedLayerStateRendering.DepthSwapchain, UnusedActualFormat);
 
-		// Image will be acquired by the renderer if supported, if not we acquire it ahead of time here
-		if (!bIsAcquireOnAnyThreadSupported)
+	// Image will be acquired by the renderer if supported, if not we acquire it ahead of time here
+	if (!bIsAcquireOnAnyThreadSupported)
+	{
+		FRHICommandListImmediate& RHICmdList = FRHICommandListImmediate::Get();
+		FXRSwapChainPtr& DepthSwapchain = PipelinedLayerStateRendering.DepthSwapchain;
+		RHICmdList.EnqueueLambda([this, DepthSwapchain](FRHICommandListImmediate&)
 		{
-			FRHICommandListImmediate& RHICmdList = FRHICommandListImmediate::Get();
-			RHICmdList.EnqueueLambda([this, DepthSwapchain](FRHICommandListImmediate&)
+			RenderBridge->RunOnRHISubmissionThread([DepthSwapchain]()
 			{
-				RenderBridge->RunOnRHISubmissionThread([DepthSwapchain]()
-				{
-					DepthSwapchain->IncrementSwapChainIndex_RHIThread();
-				});
+				DepthSwapchain->IncrementSwapChainIndex_RHIThread();
 			});
-			RHICmdList.ImmediateFlush(EImmediateFlushType::FlushRHIThread);
-		}
+		});
+		RHICmdList.ImmediateFlush(EImmediateFlushType::FlushRHIThread);
 	}
 }
 
@@ -2775,6 +3071,83 @@ bool FOpenXRHMD::AllocateDepthTexture(uint32 Index, uint32 SizeX, uint32 SizeY,
 	return true;
 }
 
+bool FOpenXRHMD::GetRecommendedMotionVectorTextureSize(FIntPoint& OutTextureSize)
+{
+	if (RecommendedMotionVectorTextureSize.X == 0 || RecommendedMotionVectorTextureSize.Y == 0)
+	{
+		OutTextureSize = FIntPoint();
+		return false;
+	}
+	
+	OutTextureSize = RecommendedMotionVectorTextureSize;
+	return true;
+}
+
+bool FOpenXRHMD::GetMotionVectorTexture(uint32 Index, const FIntPoint& Size, uint8 Format, uint32 NumMips, ETextureCreateFlags Flags, FTextureRHIRef& OutTexture, uint32 NumSamples)
+{
+	check(IsInRenderingThread());
+
+	const FXRSwapChainPtr& MotionVectorSwapchain = PipelinedLayerStateRendering.MotionVectorSwapchain;
+	if (MotionVectorSwapchain == nullptr)
+	{
+		return false;
+	}
+
+	const FRHITexture* const MotionVectorSwapchainTexture = MotionVectorSwapchain->GetTexture2DArray() ? MotionVectorSwapchain->GetTexture2DArray() : MotionVectorSwapchain->GetTexture2D();
+	const FRHITextureDesc& MotionVectorSwapchainDesc = MotionVectorSwapchainTexture->GetDesc();
+		
+	if (Size != MotionVectorSwapchainDesc.Extent)
+	{
+		// We don't yet support motion vector texture sizes other than the recommended size
+		return false;
+	}
+
+	// Sample count, mip count and size should be known at AllocateRenderTargetTexture time
+	check(EnumHasAllFlags(MotionVectorSwapchainDesc.Flags, Flags));
+	check(MotionVectorSwapchainDesc.Format == Format);
+	check(MotionVectorSwapchainDesc.NumMips == FMath::Max(NumMips, 1u));
+	check(MotionVectorSwapchainDesc.NumSamples == NumSamples);
+		
+	OutTexture = MotionVectorSwapchain->GetTextureRef();
+
+	PipelinedLayerStateRendering.LayerStateFlags |= EOpenXRLayerStateFlags::SubmitMotionVectorLayer;
+
+	return true;
+}
+
+bool FOpenXRHMD::GetMotionVectorDepthTexture(uint32 Index, const FIntPoint& Size, uint8 Format, uint32 NumMips, ETextureCreateFlags Flags, FTextureRHIRef& OutTexture, uint32 NumSamples)
+{
+	check(IsInRenderingThread());
+
+	const FXRSwapChainPtr& MotionVectorDepthSwapchain = PipelinedLayerStateRendering.MotionVectorDepthSwapchain;
+	if (MotionVectorDepthSwapchain == nullptr)
+	{
+		return false;
+	}
+
+	const FRHITexture* const MotionVectorDepthSwapchainTexture = MotionVectorDepthSwapchain->GetTexture2DArray() ? MotionVectorDepthSwapchain->GetTexture2DArray() : MotionVectorDepthSwapchain->GetTexture2D();
+	const FRHITextureDesc& MotionVectorDepthSwapchainDesc = MotionVectorDepthSwapchainTexture->GetDesc();
+
+	if (Size != MotionVectorDepthSwapchainDesc.Extent)
+	{
+		// We don't yet support motion vector depth texture sizes other than the recommended size
+		return false;
+	}
+
+	// Sample count, mip count and size should be known at AllocateRenderTargetTexture time
+	check(EnumHasAllFlags(MotionVectorDepthSwapchainDesc.Flags, Flags));
+	check(MotionVectorDepthSwapchainDesc.Format == Format);
+	check(MotionVectorDepthSwapchainDesc.NumMips == FMath::Max(NumMips, 1u));
+	check(MotionVectorDepthSwapchainDesc.NumSamples == NumSamples);
+
+	OutTexture = MotionVectorDepthSwapchain->GetTextureRef();
+
+	// Our currently supported extensions always expect motion vectors + motion vector depth, can't do just one
+	PipelinedLayerStateRendering.LayerStateFlags |= EOpenXRLayerStateFlags::SubmitMotionVectorLayer;
+
+	return true;
+}
+
 void FOpenXRLayer::FPerEyeTextureData::ConfigureSwapchain(XrSession Session, FOpenXRRenderBridge* RenderBridge, FTextureRHIRef InTexture, bool bInStaticSwapchain)
 {
 	const bool bNewTexture = Texture != InTexture;
@@ -3281,8 +3654,13 @@ void FOpenXRHMD::OnBeginRendering_RenderThread(FRDGBuilder& GraphBuilder, FScene
 
 		FXRSwapChainPtr ColorSwapchain = PipelinedLayerStateRendering.ColorSwapchain;
 		FXRSwapChainPtr DepthSwapchain = PipelinedLayerStateRendering.DepthSwapchain;
+
+		// These swapchains are only present if we are supporting frame synthesis/spacewarp.
+		// Always sanity check before using them.
+		FXRSwapChainPtr MotionVectorSwapchain = PipelinedLayerStateRendering.MotionVectorSwapchain;
+		FXRSwapChainPtr MotionVectorDepthSwapchain = PipelinedLayerStateRendering.MotionVectorDepthSwapchain;
+
 		// This swapchain might not be present depending on the platform support for stereo layers.
-		// Always check for sanity before using it.
 		FXRSwapChainPtr EmulationSwapchain = PipelinedLayerStateRendering.EmulatedLayerState.EmulationSwapchain;
 
 		if (bFoveationExtensionSupported && FBFoveationImageGenerator.IsValid())
@@ -3293,12 +3671,12 @@ void FOpenXRHMD::OnBeginRendering_RenderThread(FRDGBuilder& GraphBuilder, FScene
 
 		UE_LOG(LogHMD, VeryVerbose, TEXT("%s WF_%i EnqueueLambda OnBeginRendering_RHIThread"), HMDThreadString(), PipelinedFrameStateRendering.WaitCount);
 		// For now, leaving swapchain acquisition outside of an RDG pass to make sure the swapchain is acquired early enough.
-		GraphBuilder.RHICmdList.EnqueueLambda([this, FrameState = PipelinedFrameStateRendering, ColorSwapchain, DepthSwapchain, EmulationSwapchain](FRHICommandListImmediate& InRHICmdList)
+		GraphBuilder.RHICmdList.EnqueueLambda([this, FrameState = PipelinedFrameStateRendering, ColorSwapchain, DepthSwapchain, MotionVectorSwapchain, MotionVectorDepthSwapchain, EmulationSwapchain](FRHICommandListImmediate& InRHICmdList)
 		{
 			IRHICommandContext* Context = &InRHICmdList.GetContext();
-			RenderBridge->RunOnRHISubmissionThread([this, FrameState = PipelinedFrameStateRendering, ColorSwapchain, DepthSwapchain, EmulationSwapchain, Context]()
+			RenderBridge->RunOnRHISubmissionThread([this, FrameState = PipelinedFrameStateRendering, ColorSwapchain, DepthSwapchain, MotionVectorSwapchain, MotionVectorDepthSwapchain, EmulationSwapchain, Context]()
 			{
-				OnBeginRendering_RHIThread(*Context, FrameState, ColorSwapchain, DepthSwapchain, EmulationSwapchain);
+				OnBeginRendering_RHIThread(*Context, FrameState, ColorSwapchain, DepthSwapchain, MotionVectorSwapchain, MotionVectorDepthSwapchain, EmulationSwapchain);
 			});
 		});
 	}
@@ -3760,7 +4138,7 @@ void FOpenXRHMD::RequestExitApp()
 	}
 }
 
-void FOpenXRHMD::OnBeginRendering_RHIThread(IRHICommandContext& RHICmdContext, const FPipelinedFrameState& InFrameState, FXRSwapChainPtr ColorSwapchain, FXRSwapChainPtr DepthSwapchain, FXRSwapChainPtr EmulationSwapchain)
+void FOpenXRHMD::OnBeginRendering_RHIThread(IRHICommandContext& RHICmdContext, const FPipelinedFrameState& InFrameState, FXRSwapChainPtr ColorSwapchain, FXRSwapChainPtr DepthSwapchain, FXRSwapChainPtr MotionVectorSwapchain, FXRSwapChainPtr MotionVectorDepthSwapchain, FXRSwapChainPtr EmulationSwapchain)
 {
 	// TODO: Add a hook to resolve discarded frames before we start a new frame.
 	UE_CLOG(bIsRendering, LogHMD, Verbose, TEXT("Discarded previous frame and started rendering a new frame."));
@@ -3818,12 +4196,13 @@ void FOpenXRHMD::OnBeginRendering_RHIThread(IRHICommandContext& RHICmdContext, c
 		// TODO Possibly move these Waits to SetFinalViewRect??
 		PipelinedLayerStateRHI.ColorSwapchain = ColorSwapchain;
 		PipelinedLayerStateRHI.DepthSwapchain = DepthSwapchain;
+		PipelinedLayerStateRHI.MotionVectorSwapchain = MotionVectorSwapchain;
+		PipelinedLayerStateRHI.MotionVectorDepthSwapchain = MotionVectorDepthSwapchain;
 		PipelinedLayerStateRHI.EmulatedLayerState.EmulationSwapchain = EmulationSwapchain;
 
 		// We need a new swapchain image unless we've already acquired one for rendering
 		if (!bIsRendering && ColorSwapchain)
 		{
-			TArray<XrSwapchain> Swapchains;
 			ColorSwapchain->WaitCurrentImage_RHIThread(OPENXR_SWAPCHAIN_WAIT_TIMEOUT);
 			if (!bIsAcquireOnAnyThreadSupported)
 			{
@@ -3845,6 +4224,24 @@ void FOpenXRHMD::OnBeginRendering_RHIThread(IRHICommandContext& RHICmdContext, c
 					EmulationSwapchain->IncrementSwapChainIndex_RHIThread();
 				}
 			}
+			if (MotionVectorSwapchain)
+			{
+				MotionVectorSwapchain->WaitCurrentImage_RHIThread(OPENXR_SWAPCHAIN_WAIT_TIMEOUT);
+				if (!bIsAcquireOnAnyThreadSupported)
+				{
+					MotionVectorSwapchain->IncrementSwapChainIndex_RHIThread();
+				}
+			}
+			if (MotionVectorDepthSwapchain)
+			{
+				checkf(MotionVectorSwapchain, TEXT("Motion vector depth swapchain should not exist without a motion vector swapchain."));
+
+				MotionVectorDepthSwapchain->WaitCurrentImage_RHIThread(OPENXR_SWAPCHAIN_WAIT_TIMEOUT);
+				if (!bIsAcquireOnAnyThreadSupported)
+				{
+					MotionVectorDepthSwapchain->IncrementSwapChainIndex_RHIThread();
+				}
+			}
 		}
 
 		bIsRendering = true;
@@ -3892,6 +4289,14 @@ void FOpenXRHMD::OnFinishRendering_RHIThread(IRHICommandContext& RHICmdContext)
 		{
 			PipelinedLayerStateRHI.DepthSwapchain->ReleaseCurrentImage_RHIThread(RHICommandContextIfRequired);
 		}
+		if (PipelinedLayerStateRHI.MotionVectorSwapchain)
+		{
+			PipelinedLayerStateRHI.MotionVectorSwapchain->ReleaseCurrentImage_RHIThread(RHICommandContextIfRequired);
+		}
+		if (PipelinedLayerStateRHI.MotionVectorDepthSwapchain)
+		{
+			PipelinedLayerStateRHI.MotionVectorDepthSwapchain->ReleaseCurrentImage_RHIThread(RHICommandContextIfRequired);
+		}
 		if (PipelinedLayerStateRHI.EmulatedLayerState.EmulationSwapchain)
 		{
 			PipelinedLayerStateRHI.EmulatedLayerState.EmulationSwapchain->ReleaseCurrentImage_RHIThread(RHICommandContextIfRequired);
@@ -3914,6 +4319,7 @@ void FOpenXRHMD::OnFinishRendering_RHIThread(IRHICommandContext& RHICmdContext)
 			Layer.space = PipelinedFrameStateRHI.TrackingSpace->Handle;
 			Layer.viewCount = PipelinedLayerStateRHI.ProjectionLayers.Num();
 			Layer.views = PipelinedLayerStateRHI.ProjectionLayers.GetData();
+
 			Headers.Add(reinterpret_cast<XrCompositionLayerBaseHeader*>(&Layer));
 
 			if (CVarOpenXRInvertAlpha.GetValueOnRenderThread())
diff --git a/Engine/Plugins/Runtime/OpenXR/Source/OpenXRHMD/Private/OpenXRHMD.h b/Engine/Plugins/Runtime/OpenXR/Source/OpenXRHMD/Private/OpenXRHMD.h
--- a/Engine/Plugins/Runtime/OpenXR/Source/OpenXRHMD/Private/OpenXRHMD.h
+++ b/Engine/Plugins/Runtime/OpenXR/Source/OpenXRHMD/Private/OpenXRHMD.h
@@ -198,6 +198,7 @@ public:
 		SubmitBackgroundLayer = (1u << 1),
 		SubmitDepthLayer = (1u << 2),
 		SubmitEmulatedFaceLockedLayer = (1u << 3),
+		SubmitMotionVectorLayer = (1u << 4),
 	};
 	FRIEND_ENUM_CLASS_FLAGS(EOpenXRLayerStateFlags);
 
@@ -208,11 +209,19 @@ public:
 		TArray<XrCompositionLayerDepthInfoKHR> DepthLayers;
 		TArray<XrCompositionLayerDepthTestFB> CompositionDepthTestLayers;
 
+		TArray<XrFrameSynthesisInfoEXT> FrameSynthesisLayers;
+		TArray<XrCompositionLayerSpaceWarpInfoFB> SpaceWarpLayers;
+
 		TArray<XrSwapchainSubImage> ColorImages;
 		TArray<XrSwapchainSubImage> DepthImages;
 
+		TArray<XrSwapchainSubImage> MotionVectorImages;
+		TArray<XrSwapchainSubImage> MotionVectorDepthImages;
+
 		FXRSwapChainPtr ColorSwapchain;
 		FXRSwapChainPtr DepthSwapchain;
+		FXRSwapChainPtr MotionVectorSwapchain;
+		FXRSwapChainPtr MotionVectorDepthSwapchain;
 		TArray<FXRSwapChainPtr> NativeOverlaySwapchains;
 
 		FEmulatedLayerState EmulatedLayerState;
@@ -428,6 +437,9 @@ public:
 	virtual int32 AcquireColorTexture() override final;
 	virtual int32 AcquireDepthTexture() override final;
 	virtual bool AllocateDepthTexture(uint32 Index, uint32 SizeX, uint32 SizeY, uint8 Format, uint32 NumMips, ETextureCreateFlags InTexFlags, ETextureCreateFlags TargetableTextureFlags, FTextureRHIRef& OutTargetableTexture, FTextureRHIRef& OutShaderResourceTexture, uint32 NumSamples = 1) override final;
+	virtual bool GetRecommendedMotionVectorTextureSize(FIntPoint& OutTextureSize) override;
+	virtual bool GetMotionVectorTexture(uint32 Index, const FIntPoint& Size, uint8 Format, uint32 NumMips, ETextureCreateFlags Flags, FTextureRHIRef& OutTexture, uint32 NumSamples = 1) override;
+	virtual bool GetMotionVectorDepthTexture(uint32 Index, const FIntPoint& Size, uint8 Format, uint32 NumMips, ETextureCreateFlags Flags, FTextureRHIRef& OutTexture, uint32 NumSamples = 1) override;
 	virtual bool ReconfigureForShaderPlatform(EShaderPlatform NewShaderPlatform) override;
 	virtual EPixelFormat GetActualColorSwapchainFormat() const override { return static_cast<EPixelFormat>(LastActualColorSwapchainFormat); }
 
@@ -456,7 +468,7 @@ public:
 	virtual ~FOpenXRHMD();
 
 	void OnBeginSimulation_GameThread();
-	void OnBeginRendering_RHIThread(IRHICommandContext& RHICmdContext, const FPipelinedFrameState& InFrameState, FXRSwapChainPtr ColorSwapchain, FXRSwapChainPtr DepthSwapchain, FXRSwapChainPtr EmulationSwapchain);
+	void OnBeginRendering_RHIThread(IRHICommandContext& RHICmdContext, const FPipelinedFrameState& InFrameState, FXRSwapChainPtr ColorSwapchain, FXRSwapChainPtr DepthSwapchain, FXRSwapChainPtr MotionVectorSwapchain, FXRSwapChainPtr MotionVectorDepthSwapchain, FXRSwapChainPtr EmulationSwapchain);
 	void OnFinishRendering_RHIThread(IRHICommandContext& RHICmdContext);
 
 	/** IOpenXRHMD */
@@ -491,6 +503,8 @@ public:
 	/** Provides the full asymmetric FOV information for the view, which allows proper centering of screen-based techniques based on forward eye gaze */
 	void GetFullFOVInformation(TArray<XrFovf>& FovInfos) const override;
 
+	virtual bool AllocateSwapchainTextures_RenderThread(const FOpenXRSwapchainProperties& InSwapchainProperies, FXRSwapChainPtr& InOutSwapchain, uint8& OutActualFormat) override;
+
 	/** Returns shader platform the plugin is currently configured for, in the editor it can change due to preview platforms. */
 	EShaderPlatform GetConfiguredShaderPlatform() const { check(ConfiguredShaderPlatform != EShaderPlatform::SP_NumPlatforms); return ConfiguredShaderPlatform; }
 	FOpenXRSwapchain* GetColorSwapchain_RenderThread();
@@ -607,6 +621,9 @@ private:
 	XrColor4f				LayerColorBias;
 	bool					bCompositionLayerColorScaleBiasSupported;
 	bool					bxrGetSystemPropertiesSuccessful;
+	
+	FTransform				FrameSynthesisLastTrackingToWorld;
+	FIntPoint				RecommendedMotionVectorTextureSize;
 };
 
 ENUM_CLASS_FLAGS(FOpenXRHMD::EOpenXRLayerStateFlags);
diff --git a/Engine/Plugins/Runtime/OpenXR/Source/OpenXRHMD/Private/OpenXRHMDModule.cpp b/Engine/Plugins/Runtime/OpenXR/Source/OpenXRHMD/Private/OpenXRHMDModule.cpp
--- a/Engine/Plugins/Runtime/OpenXR/Source/OpenXRHMD/Private/OpenXRHMDModule.cpp
+++ b/Engine/Plugins/Runtime/OpenXR/Source/OpenXRHMD/Private/OpenXRHMDModule.cpp
@@ -589,15 +589,26 @@ bool FOpenXRHMDModule::GetOptionalExtensions(TArray<const ANSICHAR*>& OutExtensi
 #endif
 
 	// These extensions are mutually incompatible because XR_FB_COMPOSITION_LAYER_ALPHA_BLEND disables non-opaque environment blend modes
-	// XR_EXT_COMPOSITION_LAYER_INVERTED_ALPHA is the more modern and platform-agnostic option and is preferred
+	// XR_EXT_COMPOSITION_LAYER_INVERTED_ALPHA is the more modern and platform-agnostic option and is preferred when available
 	if (AvailableExtensions.Contains(XR_EXT_COMPOSITION_LAYER_INVERTED_ALPHA_EXTENSION_NAME))
 	{
 		OutExtensions.Add(XR_EXT_COMPOSITION_LAYER_INVERTED_ALPHA_EXTENSION_NAME);
 	}
-	else
+	else if (AvailableExtensions.Contains(XR_FB_COMPOSITION_LAYER_ALPHA_BLEND_EXTENSION_NAME))
 	{
 		OutExtensions.Add(XR_FB_COMPOSITION_LAYER_ALPHA_BLEND_EXTENSION_NAME);
 	}
+
+	// These extensions both perform frame synthesis, and require providing motion vector and motion vector depth swapchains at a requested size
+	// XR_EXT_FRAME_SYNTHESIS is not currently functioning properly on Quest, so XR_FB_SPACE_WARP is preferred when available
+	if (AvailableExtensions.Contains(XR_FB_SPACE_WARP_EXTENSION_NAME))
+	{
+		OutExtensions.Add(XR_FB_SPACE_WARP_EXTENSION_NAME);
+	}
+	else if (AvailableExtensions.Contains(XR_EXT_FRAME_SYNTHESIS_EXTENSION_NAME))
+	{
+		OutExtensions.Add(XR_EXT_FRAME_SYNTHESIS_EXTENSION_NAME);
+	}
 	
 	OutExtensions.Add(XR_EXT_LOCAL_FLOOR_EXTENSION_NAME);
 	OutExtensions.Add(XR_EXT_USER_PRESENCE_EXTENSION_NAME);
diff --git a/Engine/Plugins/Runtime/OpenXR/Source/OpenXRHMD/Public/IOpenXRHMD.h b/Engine/Plugins/Runtime/OpenXR/Source/OpenXRHMD/Public/IOpenXRHMD.h
--- a/Engine/Plugins/Runtime/OpenXR/Source/OpenXRHMD/Public/IOpenXRHMD.h
+++ b/Engine/Plugins/Runtime/OpenXR/Source/OpenXRHMD/Public/IOpenXRHMD.h
@@ -7,8 +7,23 @@
 // See also IXRTrackingSystem::GetIOpenXRHMD().
 
 #include "CoreMinimal.h"
+#include "XRSwapChain.h"
 #include <openxr/openxr.h>
 
+struct FOpenXRSwapchainProperties
+{
+	const FStringView DebugName;
+	uint8 Format;
+	uint32 SizeX;
+	uint32 SizeY;
+	uint32 ArraySize;
+	uint32 NumMips;
+	uint32 NumSamples;
+	ETextureCreateFlags CreateFlags;
+	const FClearValueBinding& ClearValueBinding;
+	ETextureCreateFlags AuxiliaryCreateFlags = ETextureCreateFlags::None;
+};
+
 enum class EOpenXRAPIVersion : uint32;
 
 class IOpenXRHMD
@@ -46,5 +61,7 @@ public:
 	OPENXRHMD_API virtual TArray<XrEnvironmentBlendMode> GetSupportedEnvironmentBlendModes() const = 0;
 
 	OPENXRHMD_API virtual void GetFullFOVInformation(TArray<XrFovf>& FovInfos) const = 0;
+
+	OPENXRHMD_API virtual bool AllocateSwapchainTextures_RenderThread(const FOpenXRSwapchainProperties& InSwapchainProperies, FXRSwapChainPtr& InOutSwapchain, uint8& OutActualFormat) = 0;
 };
 
diff --git a/Engine/Shaders/Private/GpuSkinVertexFactory.ush b/Engine/Shaders/Private/GpuSkinVertexFactory.ush
--- a/Engine/Shaders/Private/GpuSkinVertexFactory.ush
+++ b/Engine/Shaders/Private/GpuSkinVertexFactory.ush
@@ -614,7 +614,13 @@ float3 SkinPreviousPosition( FVertexFactoryInput Input, FVertexFactoryIntermedia
 	if (ResolvedView.FrameCounter == GPUSkinVF.BoneUpdatedFrameNumber && ResolvedView.WorldIsPaused == 0)
 #endif
 	{
+#if defined(IS_VELOCITY_SHADER) && STEREO_MOTION_VECTORS
+		FBoneMatrix PrevBlendMatrix = CalcPreviousBoneMatrix(Input);
+		if (IsSmoothRotation(PrevBlendMatrix._11_21_31, PrevBlendMatrix._12_22_32, BlendMatrix._11_21_31, BlendMatrix._12_22_32))
+			BlendMatrix = PrevBlendMatrix;
+#else
 		BlendMatrix = CalcPreviousBoneMatrix(Input, Intermediates.SceneData.InstanceData);
+#endif
 	}
 
 	// Note the use of mul(Matrix,Vector), bone matrices are stored transposed
diff --git a/Engine/Shaders/Private/LocalVertexFactory.ush b/Engine/Shaders/Private/LocalVertexFactory.ush
--- a/Engine/Shaders/Private/LocalVertexFactory.ush
+++ b/Engine/Shaders/Private/LocalVertexFactory.ush
@@ -1267,6 +1267,13 @@ float4 VertexFactoryGetPreviousWorldPosition(FVertexFactoryInput Input, FVertexF
 	PrevLocalPosition = mul(Input.Position, InstanceTransform);
 #endif
 	FDFMatrix PreviousLocalToWorld = GetInstanceData(Intermediates).PrevLocalToWorld;
+
+#if defined(IS_VELOCITY_SHADER) && STEREO_MOTION_VECTORS
+	FLWCMatrix CurrentLocalToWorld = GetInstanceData(Intermediates).LocalToWorld;
+	if (!IsSmoothRotation(PreviousLocalToWorld.M._11_21_31, PreviousLocalToWorld.M._12_22_32, CurrentLocalToWorld.M._11_21_31, CurrentLocalToWorld.M._12_22_32))
+		PreviousLocalToWorld = CurrentLocalToWorld;
+#endif
+
 	return TransformPreviousLocalPositionToTranslatedWorld(PrevLocalPosition.xyz, PreviousLocalToWorld);
 }
 
diff --git a/Engine/Shaders/Private/VelocityCommon.ush b/Engine/Shaders/Private/VelocityCommon.ush
--- a/Engine/Shaders/Private/VelocityCommon.ush
+++ b/Engine/Shaders/Private/VelocityCommon.ush
@@ -21,6 +21,19 @@ float3 Calculate3DVelocityBase(float4 PackedVelocityA, float4 PackedVelocityC, f
 	return Velocity;
 }
 
+#if STEREO_MOTION_VECTORS
+// If an object is rotatating too fast, the motion vector direction can get messed up and become unsuitable for OpenXR frame synthesis
+// Here, when rotation speed is above a threshold (10 degrees), we'll return false and disable motion vectors by replacing PrevTransform with CurrentTransform
+// To save shader performance, we only detect rotation on 2 axes
+bool IsSmoothRotation( float3 PrevAxis1, float3 PrevAxis2, float3 CurrAxis1, float3 CurrAxis2 )
+{
+	const float angleThreshold = 0.984f; // cos(10 degree)
+	float2 AngleDot = float2( dot(normalize(PrevAxis1), normalize(CurrAxis1)),
+	  						  dot(normalize(PrevAxis2), normalize(CurrAxis2)));
+	return all(AngleDot > angleThreshold);
+}
+#endif
+
 float3 Calculate3DVelocity(float4 PackedVelocityA, float4 PackedVelocityC)
 {
 	return Calculate3DVelocityBase(PackedVelocityA, PackedVelocityC, ResolvedView.TemporalAAJitter.xy, ResolvedView.TemporalAAJitter.zw);
@@ -30,4 +43,4 @@ float3 Calculate3DVelocity(float4 PackedVelocityA, float4 PackedVelocityC)
 float3 Calculate3DVelocityOffset(float4 PackedVelocityA, float4 PackedVelocityC)
 {
 	return Calculate3DVelocityBase(PackedVelocityA, PackedVelocityC, ResolvedView.TemporalAAJitter.xy, ResolvedView.TemporalAAJitter.xy);
-}
\ No newline at end of file
+}
diff --git a/Engine/Shaders/Private/VelocityShader.usf b/Engine/Shaders/Private/VelocityShader.usf
--- a/Engine/Shaders/Private/VelocityShader.usf
+++ b/Engine/Shaders/Private/VelocityShader.usf
@@ -22,6 +22,10 @@ static const float GDepthBias = 0.001f;
 struct FVelocityInterpsVSToPS
 {
 #if !TRANSLUCENCY_VELOCITY_FROM_DEPTH
+#if STEREO_MOTION_VECTORS
+	// float3(ScreenPos.xyzw)
+	float4 PackedVelocityA : TEXCOORD6;
+#endif
 	// float4(PrevScreenPos.xyzw)
 	float4 PackedVelocityC : TEXCOORD7;
 #endif
@@ -83,6 +87,10 @@ void MainVertexShader(
 		float4 PrevScreenPosObj = mul(float4(PrevTranslatedWorldPosition.xyz, 1), ResolvedView.PrevTranslatedWorldToClip);
 		Output.Position = INVARIANT(mul(float4(TranslatedWorldPosition.xyz, 1), ResolvedView.TranslatedWorldToClip));
 
+#if STEREO_MOTION_VECTORS
+		Output.Interps.PackedVelocityA = Output.Position;
+#endif
+
 #if USE_GLOBAL_CLIP_PLANE
 		OutGlobalClipPlaneDistance = dot(ResolvedView.GlobalClippingPlane, float4(TranslatedWorldPosition.xyz, 1));
 #endif
@@ -132,7 +140,13 @@ void MainPixelShader(
 	FMaterialPixelParameters MaterialParameters = GetMaterialPixelParameters(Input.FactoryInterpolants, SvPosition);
 	FPixelMaterialInputs PixelMaterialInputs;
 
+#if STEREO_MOTION_VECTORS
+	float4 ScreenPosition = Input.PackedVelocityA;
+#else
 	float4 ScreenPosition = SvPositionToResolvedScreenPosition(SvPosition);
+#endif
+
+
 #if USE_WORLD_POSITION_EXCLUDING_SHADER_OFFSETS && !IS_NANITE_PASS
 	float3 TranslatedWorldPosition = SvPositionToResolvedTranslatedWorld(SvPosition);
 	CalcMaterialParametersEx(MaterialParameters, PixelMaterialInputs, SvPosition, ScreenPosition, bIsFrontFace, TranslatedWorldPosition, Input.PixelPositionExcludingWPO);
@@ -231,7 +245,12 @@ void MainPixelShader(
 		Velocity = Calculate3DVelocity(ScreenPosition, Input.PackedVelocityC);
 	#endif // !VELOCITY_CLIPPED_DEPTH_PASS
 
+#if STEREO_MOTION_VECTORS
+	OutColor.xyz = Velocity;
+	OutColor.w = 0.0f;
+#else
 	OutColor = EncodeVelocityToTexture(Velocity, bHasPixelAnimationFlag, TemporalResponsiveness);
+#endif
 #else
 	// We only want velocity generated from camera movement and depth so write zero which is the special clear color value.
 	OutColor = 0;
diff --git a/Engine/Source/Runtime/Engine/Private/ModelRender.cpp b/Engine/Source/Runtime/Engine/Private/ModelRender.cpp
--- a/Engine/Source/Runtime/Engine/Private/ModelRender.cpp
+++ b/Engine/Source/Runtime/Engine/Private/ModelRender.cpp
@@ -561,7 +561,7 @@ public:
 		}
 		Result.bShadowRelevance = IsShadowCast(View);
 		MaterialRelevance.SetPrimitiveViewRelevance(Result);
-		Result.bVelocityRelevance = DrawsVelocity() && Result.bOpaque && Result.bRenderInMainPass;
+		Result.bVelocityRelevance = (VelocityIncludeStationaryPrimitives(View->GetShaderPlatform()) || DrawsVelocity()) && Result.bOpaque && Result.bRenderInMainPass;
 		return Result;
 	}
 
diff --git a/Engine/Source/Runtime/Engine/Private/StaticMeshSceneProxy.cpp b/Engine/Source/Runtime/Engine/Private/StaticMeshSceneProxy.cpp
--- a/Engine/Source/Runtime/Engine/Private/StaticMeshSceneProxy.cpp
+++ b/Engine/Source/Runtime/Engine/Private/StaticMeshSceneProxy.cpp
@@ -2353,7 +2353,7 @@ FPrimitiveViewRelevance FStaticMeshSceneProxy::GetViewRelevance(const FSceneView
 		Result.bOpaque = true;
 	}
 
-	Result.bVelocityRelevance = DrawsVelocity() && Result.bOpaque && Result.bRenderInMainPass;
+	Result.bVelocityRelevance = (VelocityIncludeStationaryPrimitives(View->GetShaderPlatform()) || DrawsVelocity()) && Result.bOpaque && Result.bRenderInMainPass;
 
 	return Result;
 }
diff --git a/Engine/Source/Runtime/Engine/Public/ShowFlags.h b/Engine/Source/Runtime/Engine/Public/ShowFlags.h
--- a/Engine/Source/Runtime/Engine/Public/ShowFlags.h
+++ b/Engine/Source/Runtime/Engine/Public/ShowFlags.h
@@ -482,6 +482,7 @@ private:
 		SetHighResScreenshotMask(false);
 		SetHMDDistortion(false); // only effective if a HMD is in use
 		SetAlphaInvert(false);
+		SetStereoMotionVectors(false);
 		SetStereoRendering(false);
 		SetDistanceCulledPrimitives(InitMode == ESFIM_Editor);
 		SetVisualizeLightCulling(false);
diff --git a/Engine/Source/Runtime/Engine/Public/ShowFlagsValues.inl b/Engine/Source/Runtime/Engine/Public/ShowFlagsValues.inl
--- a/Engine/Source/Runtime/Engine/Public/ShowFlagsValues.inl
+++ b/Engine/Source/Runtime/Engine/Public/ShowFlagsValues.inl
@@ -482,5 +482,8 @@ SHOWFLAG_ALWAYS_ACCESSIBLE(AllowPrimitiveAlphaHoldout, SFG_Advanced, NSLOCTEXT("
 /** Invert the alpha channel at tne end of post processing. */
 SHOWFLAG_ALWAYS_ACCESSIBLE(AlphaInvert, SFG_PostProcess, NSLOCTEXT("UnrealEd", "AlphaInvertSF", "Invert Alpha Channel"))
 
+/** Directly render motion vectors and motion vector depth for OpenXR extensions when normal velocity rendering is unsupported. Motion vectors/depth must be in Vulkan NDC-space and can be requested at any size. */
+SHOWFLAG_ALWAYS_ACCESSIBLE(StereoMotionVectors, SFG_Normal, NSLOCTEXT("UnrealEd", "MotionVectorsSF", "Render Stereo Motion Vectors for OpenXR"))
+
 #undef SHOWFLAG_ALWAYS_ACCESSIBLE
 #undef SHOWFLAG_FIXED_IN_SHIPPING
diff --git a/Engine/Source/Runtime/Engine/Public/StereoRenderTargetManager.h b/Engine/Source/Runtime/Engine/Public/StereoRenderTargetManager.h
--- a/Engine/Source/Runtime/Engine/Public/StereoRenderTargetManager.h
+++ b/Engine/Source/Runtime/Engine/Public/StereoRenderTargetManager.h
@@ -138,5 +138,42 @@ public:
 	 */
 	virtual bool ReconfigureForShaderPlatform(EShaderPlatform NewShaderPlatform) { return false; };
 
+	/**
+	 * Gets the size of the motion vector texture recommended by the stereo device.
+	 * Allocating a motion vector texture of any other size may be detrimental to performance or be unsupported by the stereo device.
+	 *
+	 * @param OutTextureSize			(out) The size of the motion vector texture.
+	 * @return							true, if the stereo device can recommend a texture size; false, if not.
+	 */
+	virtual bool GetRecommendedMotionVectorTextureSize(FIntPoint& OutTextureSize) { return false; }
+	
+	/**
+	 * Allocates a motion vector texture.
+	 * The default implementation always returns false to indicate that the stereo device does not support motion vector textures.
+	 *
+	 * @param Size						(in) requested motion vector texture size
+	 * @param Format					(in) requested motion vector texture format
+	 * @param NumMips					(in) requested motion vector texture mips count
+	 * @param Flags						(in) requested motion vector texture flags
+	 * @param OutTexture				(out) motion vector texture resource allocated by the stereo device
+	 * @param NumSamples				(in) requested motion vector texture samples count
+	 * @return							true, if texture was allocated; false, if the stereo device can't provide a texture.
+	 */
+	virtual bool GetMotionVectorTexture(uint32 Index, const FIntPoint& Size, uint8 Format, uint32 NumMips, ETextureCreateFlags Flags, FTextureRHIRef& OutTexture, uint32 NumSamples = 1) { return false; }
+	
+	/**
+	 * Allocates a motion vector depth texture.
+	 * The default implementation always returns false to indicate that the stereo device doesn't support motion vector depth textures.
+	 *
+	 * @param Size						(in) requested motion vector depth texture size
+	 * @param Format					(in) requested motion vector depth texture format
+	 * @param NumMips					(in) requested motion vector depth texture mips count
+	 * @param Flags						(in) requested motion vector depth texture flags
+	 * @param OutTexture				(out) motion vector depth texture resource allocated by the stereo device
+	 * @param NumSamples				(in) requested motion vector depth texture samples count
+	 * @return							true, if texture was allocated; false, if the stereo device can't provide a texture.
+	 */
+	virtual bool GetMotionVectorDepthTexture(uint32 Index, const FIntPoint& Size, uint8 Format, uint32 NumMips, ETextureCreateFlags Flags, FTextureRHIRef& OutTexture, uint32 NumSamples = 1) { return false; }
+	
 	static EPixelFormat GetStereoLayerPixelFormat() { return PF_B8G8R8A8; }
 };
diff --git a/Engine/Source/Runtime/Landscape/Private/LandscapeRender.cpp b/Engine/Source/Runtime/Landscape/Private/LandscapeRender.cpp
--- a/Engine/Source/Runtime/Landscape/Private/LandscapeRender.cpp
+++ b/Engine/Source/Runtime/Landscape/Private/LandscapeRender.cpp
@@ -2175,7 +2175,7 @@ FPrimitiveViewRelevance FLandscapeComponentSceneProxy::GetViewRelevance(const FS
 #endif
 	}
 
-	Result.bVelocityRelevance = DrawsVelocity() && Result.bOpaque && Result.bRenderInMainPass;
+	Result.bVelocityRelevance = (VelocityIncludeStationaryPrimitives(View->GetShaderPlatform()) || DrawsVelocity()) && Result.bOpaque && Result.bRenderInMainPass;
 	Result.bShadowRelevance = (GAllowLandscapeShadows > 0) && IsShadowCast(View) && View->Family->EngineShowFlags.Landscape;
 
 #if !UE_BUILD_SHIPPING
diff --git a/Engine/Source/Runtime/RenderCore/Private/RenderUtils.cpp b/Engine/Source/Runtime/RenderCore/Private/RenderUtils.cpp
--- a/Engine/Source/Runtime/RenderCore/Private/RenderUtils.cpp
+++ b/Engine/Source/Runtime/RenderCore/Private/RenderUtils.cpp
@@ -1254,8 +1254,8 @@ RENDERCORE_API bool PlatformSupportsVelocityRendering(const FStaticShaderPlatfor
 {
 	if (IsMobilePlatform(Platform))
 	{
-		// Enable velocity rendering if desktop Gen4 TAA is supported on mobile.
-		return SupportsGen4TAA(Platform);
+		// Enable velocity rendering if desktop Gen4 TAA is supported on mobile or we're performing direct OpenXR motion vector rendering for frame synthesis.
+		return SupportsGen4TAA(Platform) || PlatformSupportsOpenXRMotionVectors(Platform);
 	}
 
 	return true;
@@ -1668,6 +1668,23 @@ bool VelocityEncodeHasPixelAnimation(const FStaticShaderPlatform Platform)
 	return !IsMobilePlatform(Platform);
 }
 
+bool PlatformSupportsOpenXRMotionVectors(const FStaticShaderPlatform Platform)
+{
+	if (IsVulkanMobilePlatform(Platform) && !IsMobileDeferredShadingEnabled(Platform))
+	{
+		static FShaderPlatformCachedIniValue<bool> PerPlatformCVar(TEXT("r.Velocity.DirectlyRenderOpenXRMotionVectors"));
+		return PerPlatformCVar.Get(Platform);
+	}
+
+	return false;
+}
+
+bool VelocityIncludeStationaryPrimitives(const FStaticShaderPlatform Platform)
+{
+	// Directly rendering OpenXR motion vectors requires including all primitives in the velocity pass, even stationary primitives that would normally be excluded
+	return PlatformSupportsOpenXRMotionVectors(Platform);
+}
+
 RENDERCORE_API bool AllowTranslucencyPerObjectShadows(const FStaticShaderPlatform Platform)
 {
 	return IsFeatureLevelSupported(Platform, ERHIFeatureLevel::SM5) && GAllowTranslucencyShadowsInProject != 0;
diff --git a/Engine/Source/Runtime/RenderCore/Public/RenderUtils.h b/Engine/Source/Runtime/RenderCore/Public/RenderUtils.h
--- a/Engine/Source/Runtime/RenderCore/Public/RenderUtils.h
+++ b/Engine/Source/Runtime/RenderCore/Public/RenderUtils.h
@@ -380,6 +380,10 @@ RENDERCORE_API bool VelocityEncodeDepth(const FStaticShaderPlatform Platform);
 
 RENDERCORE_API bool VelocityEncodeHasPixelAnimation(const FStaticShaderPlatform Platform);
 
+RENDERCORE_API bool PlatformSupportsOpenXRMotionVectors(const FStaticShaderPlatform Platform);
+
+RENDERCORE_API bool VelocityIncludeStationaryPrimitives(const FStaticShaderPlatform Platform);
+
 /** Unit cube vertex buffer (VertexDeclarationFVector4) */
 RENDERCORE_API FBufferRHIRef& GetUnitCubeVertexBuffer();
 
diff --git a/Engine/Source/Runtime/Renderer/Internal/SceneTextures.h b/Engine/Source/Runtime/Renderer/Internal/SceneTextures.h
--- a/Engine/Source/Runtime/Renderer/Internal/SceneTextures.h
+++ b/Engine/Source/Runtime/Renderer/Internal/SceneTextures.h
@@ -152,6 +152,11 @@ struct FSceneTextures : public FMinimalSceneTextures
 	// Texture used by the quad overdraw debug view mode when enabled.
 	FRDGTextureRef QuadOverdraw{};
 
+	// (Mobile Stereo) Textures containing motion vectors and depth for OpenXR frame synthesis, using Vulkan-NDC format rather than UE's encoded velocity format.
+	// The size of both textures is determined independently of the view family size, via IStereoRenderTargetManager::GetRecommendedMotionVectorTextureSize.
+	FRDGTextureRef StereoMotionVectors{};
+	FRDGTextureRef StereoMotionVectorDepth{};
+
 	// Textures used to composite editor primitives. Also used by the base pass when in wireframe mode.
 #if WITH_EDITOR
 	FRDGTextureRef EditorPrimitiveColor{};
diff --git a/Engine/Source/Runtime/Renderer/Private/DepthRendering.cpp b/Engine/Source/Runtime/Renderer/Private/DepthRendering.cpp
--- a/Engine/Source/Runtime/Renderer/Private/DepthRendering.cpp
+++ b/Engine/Source/Runtime/Renderer/Private/DepthRendering.cpp
@@ -1074,7 +1074,7 @@ void FDepthPassMeshProcessor::AddMeshBatch(const FMeshBatch& RESTRICT MeshBatch,
 		{
 			if (ViewIfDynamicMeshCommand)
 			{
-				if (FOpaqueVelocityMeshProcessor::PrimitiveHasVelocityForFrame(PrimitiveSceneProxy))
+				if (FOpaqueVelocityMeshProcessor::PrimitiveHasVelocityForFrame(ShaderPlatform, PrimitiveSceneProxy))
 				{
 					checkSlow(ViewIfDynamicMeshCommand->bIsViewInfo);
 					FViewInfo* ViewInfo = (FViewInfo*)ViewIfDynamicMeshCommand;
diff --git a/Engine/Source/Runtime/Renderer/Private/MobileShadingRenderer.cpp b/Engine/Source/Runtime/Renderer/Private/MobileShadingRenderer.cpp
--- a/Engine/Source/Runtime/Renderer/Private/MobileShadingRenderer.cpp
+++ b/Engine/Source/Runtime/Renderer/Private/MobileShadingRenderer.cpp
@@ -1639,6 +1639,88 @@ void FMobileSceneRenderer::Render(FRDGBuilder& GraphBuilder, const FSceneRenderU
 	}
 }
 
+BEGIN_SHADER_PARAMETER_STRUCT(FMotionVectorPassParameters, )
+	SHADER_PARAMETER_STRUCT_INCLUDE(FViewShaderParameters, View)
+	SHADER_PARAMETER_RDG_UNIFORM_BUFFER(FSceneTextureUniformParameters, SceneTextures)
+	SHADER_PARAMETER_STRUCT_INCLUDE(FInstanceCullingDrawParams, InstanceCullingDrawParams)
+	RENDER_TARGET_BINDING_SLOTS()
+END_SHADER_PARAMETER_STRUCT()
+
+void FMobileSceneRenderer::RenderStereoMotionVectors(FRDGBuilder& GraphBuilder, FViewInfo& View, const FSceneTextures& SceneTextures, FInstanceCullingManager& InstanceCullingManager)
+{
+	SCOPED_NAMED_EVENT(RenderStereoMotionVectors, FColor::Magenta);
+
+	check(PlatformSupportsOpenXRMotionVectors(View.GetShaderPlatform()));
+
+	bool bClearedRTs = false;
+	if (View.ShouldRenderView())
+	{
+		constexpr int32 TotalMotionVectorPasses = 2;
+		for (int32 PassIndex = 0; PassIndex < TotalMotionVectorPasses; ++PassIndex)
+		{
+			const EVelocityPass VelocityPass = PassIndex == 0 ? EVelocityPass::Opaque : EVelocityPass::Translucent;
+			const EMeshPass::Type MeshPass = GetMeshPassFromVelocityPass(VelocityPass);
+			FParallelMeshDrawCommandPass* ParallelMeshPass = View.ParallelMeshDrawCommandPasses[MeshPass];
+
+			if (!ParallelMeshPass || !ParallelMeshPass->HasAnyDraw())
+			{
+				continue;
+			}
+
+			View.BeginRenderView();
+
+			FRenderTargetBindingSlots MotionVectorRenderTargets;
+			const ERenderTargetLoadAction LoadAction = bClearedRTs ? ERenderTargetLoadAction::ELoad : ERenderTargetLoadAction::EClear;
+			MotionVectorRenderTargets[0] = FRenderTargetBinding(SceneTextures.StereoMotionVectors, LoadAction);
+			MotionVectorRenderTargets.DepthStencil = FDepthStencilBinding(SceneTextures.StereoMotionVectorDepth, ERenderTargetLoadAction::EClear, FExclusiveDepthStencil::DepthWrite_StencilWrite);
+
+			// If the main view isn't multiview but the app is, need to render as a single-view multiview due to shaders
+			MotionVectorRenderTargets.MultiViewCount = View.bIsMobileMultiViewEnabled ? 2 : (View.Aspects.IsMobileMultiViewEnabled() ? 1 : 0);
+
+			MotionVectorRenderTargets.SubpassHint = ESubpassHint::None;
+			bClearedRTs = true;
+
+			EMobileSceneTextureSetupMode SetupMode = EMobileSceneTextureSetupMode::SceneVelocity;
+			FMotionVectorPassParameters* PassParameters = GraphBuilder.AllocParameters<FMotionVectorPassParameters>();
+			PassParameters->View = View.GetShaderParameters();
+			PassParameters->SceneTextures = SceneTextures.UniformBuffer;
+			PassParameters->RenderTargets = MotionVectorRenderTargets;
+
+			BuildMeshRenderingCommands(GraphBuilder, MeshPass, View, Scene->GPUScene, InstanceCullingManager, PassParameters->InstanceCullingDrawParams);
+
+			GraphBuilder.AddPass(
+				RDG_EVENT_NAME("StereoMotionVectors"),
+				PassParameters,
+				ERDGPassFlags::Raster | ERDGPassFlags::NeverMerge,
+				[this, PassParameters, MeshPass, &View, &SceneTextures](FRHICommandListImmediate& RHICmdList)
+			{
+				CSV_SCOPED_TIMING_STAT_EXCLUSIVE(RenderStereoMotionVectors);
+				SCOPED_DRAW_EVENT(RHICmdList, RenderStereoMotionVectors);
+				SCOPE_CYCLE_COUNTER(STAT_BasePassDrawTime);
+				SCOPED_GPU_MASK(RHICmdList, View.GPUMask);
+
+				if (auto* Pass = View.ParallelMeshDrawCommandPasses[MeshPass])
+				{
+					// Always use full resolution velocity buffer
+					FIntPoint VelocityBufferSize(SceneTextures.StereoMotionVectors->GetRHI()->GetSizeXYZ().X, SceneTextures.StereoMotionVectors->GetRHI()->GetSizeXYZ().Y);
+					RHICmdList.SetStereoViewport(0, 0, 0, 0, 0.0f, VelocityBufferSize.X, VelocityBufferSize.X, VelocityBufferSize.Y, VelocityBufferSize.Y, 1.0f);
+
+					Pass->Draw(RHICmdList, &PassParameters->InstanceCullingDrawParams);
+				}
+			});
+		}
+	}
+}
+
+bool FMobileSceneRenderer::HasStereoMotionVectorRenderTargets(const FSceneTextures& SceneTextures)
+{
+	if (SceneTextures.StereoMotionVectors != nullptr && SceneTextures.StereoMotionVectorDepth != nullptr)
+	{
+		return true;
+	}
+	return false;
+}
+
 FRenderTargetBindingSlots FMobileSceneRenderer::InitRenderTargetBindings_Forward(FRDGTextureRef ViewFamilyTexture, FSceneTextures& SceneTextures)
 {
 	FRDGTextureRef SceneColor = nullptr;
@@ -1728,6 +1810,13 @@ void FMobileSceneRenderer::RenderForward(FRDGBuilder& GraphBuilder, FRDGTextureR
 		SCOPED_GPU_MASK(GraphBuilder.RHICmdList, !View.IsInstancedStereoPass() ? View.GPUMask : (View.GPUMask | View.GetInstancedView()->GPUMask));
 		SCOPED_CONDITIONAL_DRAW_EVENTF(GraphBuilder.RHICmdList, EventView, RenderViews.Num() > 1, TEXT("View%d"), ViewContext.ViewIndex);
 
+		if (View.Family->EngineShowFlags.StereoMotionVectors
+			&& HasStereoMotionVectorRenderTargets(SceneTextures) 
+			&& !View.bIsSceneCapture && !View.bIsReflectionCapture)
+		{
+			RenderStereoMotionVectors(GraphBuilder, View, SceneTextures, InstanceCullingManager);
+		}
+
 		if (!ViewContext.bIsFirstView)
 		{
 			BasePassRenderTargets[0].SetLoadAction(ERenderTargetLoadAction::ELoad);
diff --git a/Engine/Source/Runtime/Renderer/Private/SceneRendering.h b/Engine/Source/Runtime/Renderer/Private/SceneRendering.h
--- a/Engine/Source/Runtime/Renderer/Private/SceneRendering.h
+++ b/Engine/Source/Runtime/Renderer/Private/SceneRendering.h
@@ -2920,6 +2920,9 @@ protected:
 	void RenderDeferredSinglePass(FRDGBuilder& GraphBuilder, FSceneTextures& SceneTextures, const FSortedLightSetSceneInfo& SortedLightSet, FMobileBasePassTextures& BasePassTextures, FInstanceCullingManager& InstanceCullingManager);
 	void RenderDeferredMultiPass(FRDGBuilder& GraphBuilder, FSceneTextures& SceneTextures, const FSortedLightSetSceneInfo& SortedLightSet, FMobileBasePassTextures& BasePassTextures, FInstanceCullingManager& InstanceCullingManager);
 
+	void RenderStereoMotionVectors(FRDGBuilder& GraphBuilder, FViewInfo& View, const FSceneTextures& SceneTextures, FInstanceCullingManager& InstanceCullingManager);
+	bool HasStereoMotionVectorRenderTargets(const FSceneTextures& SceneTextures);
+
 	void RenderAmbientOcclusion(FRDGBuilder& GraphBuilder, const FSceneTextures& SceneTextures, FRDGTextureRef ScreenSpaceShadowMaskTexture);
 
 	FRDGTextureRef RenderMobileShadowProjections(FRDGBuilder& GraphBuilder, FIntPoint Extent, bool bIsMobileMultiView);
diff --git a/Engine/Source/Runtime/Renderer/Private/SceneTextures.cpp b/Engine/Source/Runtime/Renderer/Private/SceneTextures.cpp
--- a/Engine/Source/Runtime/Renderer/Private/SceneTextures.cpp
+++ b/Engine/Source/Runtime/Renderer/Private/SceneTextures.cpp
@@ -123,6 +123,25 @@ static TRefCountPtr<FRHITexture> FindStereoDepthTexture(uint32 bSupportsXRDepth,
 	return nullptr;
 }
 
+static bool FindStereoMotionVectorTexture(FIntPoint& OutVelocityTextureSize, FTextureRHIRef& OutVelocityTexture, FTextureRHIRef& OutVelocityDepthTexture)
+{
+	if (IStereoRenderTargetManager* StereoRenderTargetManager = FindStereoRenderTargetManager())
+	{
+		if (StereoRenderTargetManager->GetRecommendedMotionVectorTextureSize(OutVelocityTextureSize))
+		{
+			if (StereoRenderTargetManager->GetMotionVectorTexture(0, OutVelocityTextureSize, PF_FloatRGBA, 1, TexCreate_RenderTargetable | TexCreate_ResolveTargetable | TexCreate_ShaderResource | TexCreate_InputAttachmentRead, OutVelocityTexture))
+			{
+				StereoRenderTargetManager->GetMotionVectorDepthTexture(0, OutVelocityTextureSize, PF_DepthStencil, 1, TexCreate_DepthStencilTargetable | TexCreate_ShaderResource | TexCreate_InputAttachmentRead, OutVelocityDepthTexture);
+				return true;
+			}
+		}
+	}
+	
+	OutVelocityTexture = nullptr;
+	OutVelocityDepthTexture = nullptr;
+	return false;
+}
+
 /** Helper class used to track and compute a suitable scene texture extent for the renderer based on history / global configuration. */
 class FSceneTextureExtentState
 {
@@ -746,6 +765,21 @@ void FSceneTextures::InitializeViewFamily(FRDGBuilder& GraphBuilder, FViewFamily
 		Desc.NumSamples = Config.NumSamples;
 		SceneTextures.DepthAux = CreateTextureMSAA(GraphBuilder, Desc, TEXT("SceneDepthAuxMS"), TEXT("SceneDepthAux"));
 	}
+
+	if (ViewFamily.EngineShowFlags.StereoMotionVectors)
+	{
+		FTextureRHIRef StereoMotionVectorRHI, StereoMotionVectorDepthRHI;
+		FIntPoint StereoMotionVectorSize;
+		if (FindStereoMotionVectorTexture(StereoMotionVectorSize, StereoMotionVectorRHI, StereoMotionVectorDepthRHI))
+		{
+			SceneTextures.StereoMotionVectors = RegisterExternalTexture(GraphBuilder, StereoMotionVectorRHI, TEXT("ExternalStereoMotionVectors"));
+			if (StereoMotionVectorDepthRHI)
+			{
+				SceneTextures.StereoMotionVectorDepth = RegisterExternalTexture(GraphBuilder, StereoMotionVectorDepthRHI, TEXT("ExternalStereoMotionVectorDepth"));
+			}
+		}
+	}
+
 #if WITH_EDITOR
 	{
 		const FRDGTextureDesc ColorDesc(FRDGTextureDesc::Create2D(Config.Extent, PF_B8G8R8A8, FClearValueBinding::Transparent, TexCreate_ShaderResource | TexCreate_RenderTargetable, 1, Config.EditorPrimitiveNumSamples));
diff --git a/Engine/Source/Runtime/Renderer/Private/SceneVisibility.cpp b/Engine/Source/Runtime/Renderer/Private/SceneVisibility.cpp
--- a/Engine/Source/Runtime/Renderer/Private/SceneVisibility.cpp
+++ b/Engine/Source/Runtime/Renderer/Private/SceneVisibility.cpp
@@ -1630,7 +1630,7 @@ void FRelevancePacket::ComputeRelevance(FDynamicPrimitiveIndexList& DynamicPrimi
 									{
 										if (ViewRelevance.bVelocityRelevance &&
 											FOpaqueVelocityMeshProcessor::PrimitiveCanHaveVelocity(View.GetShaderPlatform(), PrimitiveSceneProxy) &&
-											FOpaqueVelocityMeshProcessor::PrimitiveHasVelocityForFrame(PrimitiveSceneProxy))
+											FOpaqueVelocityMeshProcessor::PrimitiveHasVelocityForFrame(View.GetShaderPlatform(), PrimitiveSceneProxy))
 										{
 											DrawCommandPacket.AddCommandsForMesh(PrimitiveIndex, PrimitiveSceneInfo, StaticMeshRelevance, StaticMesh, CullingPayloadFlags, Scene, bCanCache, EMeshPass::Velocity);
 											bIsMeshInVelocityPass = true;
diff --git a/Engine/Source/Runtime/Renderer/Private/VelocityRendering.cpp b/Engine/Source/Runtime/Renderer/Private/VelocityRendering.cpp
--- a/Engine/Source/Runtime/Renderer/Private/VelocityRendering.cpp
+++ b/Engine/Source/Runtime/Renderer/Private/VelocityRendering.cpp
@@ -81,6 +81,14 @@ static TAutoConsoleVariable<int32> CVarVelocityPixelShaderMotionVectorWorldOffse
 	TEXT("Whether motion vector offset is supported in PS pass. Allow user to modify the motion vector per pixel.\n"),
 	ECVF_ReadOnly);
 
+static TAutoConsoleVariable<bool> CVarVelocityDirectlyRenderOpenXRMotionVectors(
+	TEXT("r.Velocity.DirectlyRenderOpenXRMotionVectors"),
+	1,
+	TEXT("If true and using the Vulkan mobile forward renderer, the engine will render velocity in the OpenXR motion vector format, at the size recommended by FOpenXRHMD::GetRecommendedMotionVectorTextureSize.\n")
+	TEXT("Because the existing scene depth cannot be used to calculate flattened velocity for stationary objects due to the likely size mismatch, this requires including ALL meshes in the velocity pass, even stationary ones which would usually be excluded.\n")
+	TEXT("This setting disables normal velocity rendering and all other features dependent on it, such as Temporal Anti-Aliasing and Motion Blur."),
+	ECVF_ReadOnly);
+
 DECLARE_GPU_DRAWCALL_STAT_NAMED(RenderVelocities, TEXT("Render Velocities"));
 
 /** Validate that deprecated CVars are no longer set. */
@@ -133,6 +141,38 @@ class FVelocityVS : public FMeshMaterialShader
 public:
 	DECLARE_SHADER_TYPE(FVelocityVS, MeshMaterial);
 
+	static bool ShouldCompilePermutation(const FMeshMaterialShaderPermutationParameters& Parameters)
+	{
+		return false;
+	}
+
+	FVelocityVS() = default;
+	FVelocityVS(const ShaderMetaType::CompiledShaderInitializerType& Initializer)
+		: FMeshMaterialShader(Initializer)
+	{}
+};
+
+bool SupportsTranslucentClippedDepth(EShaderPlatform ShaderPlatform)
+{
+	// Translucent clipped depth requires atomics on uint64. So holds the same requirement of Nanite.
+	static FShaderPlatformCachedIniValue<int32> PerPlatformCVar(TEXT("r.Velocity.OutputTranslucentClippedDepth.Supported"));
+	return FDataDrivenShaderPlatformInfo::GetSupportsNanite(ShaderPlatform) && (PerPlatformCVar.Get(ShaderPlatform) != 0);
+}
+
+enum EVelocityPassMode
+{
+	Velocity_Standard,
+	Velocity_ClippedDepth,
+	Velocity_StereoMotionVectors,
+};
+
+// Use template to allow shader pipeline binding for better performance on some platforms
+template <EVelocityPassMode PassMode>
+class TVelocityVS : public FVelocityVS
+{
+public:
+	DECLARE_SHADER_TYPE(TVelocityVS, MeshMaterial);
+
 	static bool ShouldCompilePermutation(const FMeshMaterialShaderPermutationParameters& Parameters)
 	{
 		// Compile for default material.
@@ -151,8 +191,20 @@ public:
 		const bool bModifiesMotionVectorStatus = 
 			(Parameters.MaterialParameters.bUsesMotionVectorWorldOffset && SupportsPixelShaderMotionVectorWorldOffset(Parameters.Platform)) || 
 			(Parameters.MaterialParameters.bUsesTemporalResponsiveness && SupportsTemporalResponsiveness(Parameters.Platform));
-			
-		const bool bHasPlatformSupport = PlatformSupportsVelocityRendering(Parameters.Platform);
+
+		bool bHasPlatformSupport = false;
+		if constexpr (PassMode == EVelocityPassMode::Velocity_Standard)
+		{
+			bHasPlatformSupport = PlatformSupportsVelocityRendering(Parameters.Platform);
+		}
+		else if constexpr (PassMode == EVelocityPassMode::Velocity_ClippedDepth)
+		{
+			bHasPlatformSupport = PlatformSupportsVelocityRendering(Parameters.Platform) && SupportsTranslucentClippedDepth(Parameters.Platform);
+		}
+		else if constexpr (PassMode == EVelocityPassMode::Velocity_StereoMotionVectors)
+		{
+			bHasPlatformSupport = PlatformSupportsOpenXRMotionVectors(Parameters.Platform);
+		}
 
 		/**
 		 * If we don't use base pass velocity then we may need to generate permutations for this shader. 
@@ -170,19 +222,23 @@ public:
 		return bHasPlatformSupport && !bIsNaniteFactory && (bIsDefault || bIsSeparateVelocityPassRequired || bIsSeparateVelocityPassRequiredByMaterial);
 	}
 
-	FVelocityVS() = default;
-	FVelocityVS(const ShaderMetaType::CompiledShaderInitializerType& Initializer)
-		: FMeshMaterialShader(Initializer)
+	static void ModifyCompilationEnvironment(const FMaterialShaderPermutationParameters& Parameters, FShaderCompilerEnvironment& OutEnvironment)
+	{
+		FVelocityVS::ModifyCompilationEnvironment(Parameters, OutEnvironment);
+
+		constexpr bool bIsStereoMotionVectorPass = (PassMode == EVelocityPassMode::Velocity_StereoMotionVectors);
+		OutEnvironment.SetDefine(TEXT("STEREO_MOTION_VECTORS"), bIsStereoMotionVectorPass ? 1 : 0);
+	}
+
+	TVelocityVS()
+	{
+	}
+
+	TVelocityVS(const ShaderMetaType::CompiledShaderInitializerType& Initializer)
+		: FVelocityVS(Initializer)
 	{}
 };
 
-bool SupportsTranslucentClippedDepth(EShaderPlatform ShaderPlatform)
-{
-	// Translucent clipped depth requires atomics on uint64. So holds the same requirement of Nanite.
-	static FShaderPlatformCachedIniValue<int32> PerPlatformCVar(TEXT("r.Velocity.OutputTranslucentClippedDepth.Supported"));
-	return FDataDrivenShaderPlatformInfo::GetSupportsNanite(ShaderPlatform) && (PerPlatformCVar.Get(ShaderPlatform) != 0);
-}
-
 class FVelocityPS : public FMeshMaterialShader
 {
 public:
@@ -208,12 +264,6 @@ public:
 	{}
 };
 
-enum EVelocityPassMode
-{
-	Velocity_Standard,
-	Velocity_ClippedDepth
-};
-
 // Use template to allow shader pipeline binding for better performance on some platforms
 template <EVelocityPassMode PassMode>
 class TVelocityPS : public FVelocityPS
@@ -223,29 +273,19 @@ public:
 
 	static bool ShouldCompilePermutation(const FMeshMaterialShaderPermutationParameters& Parameters)
 	{
-		if (FVelocityVS::ShouldCompilePermutation(Parameters))
-		{
-			const bool bIsVelocityClippedDepthPass = PassMode == EVelocityPassMode::Velocity_ClippedDepth;
-
-			if (bIsVelocityClippedDepthPass)
-			{
-				return SupportsTranslucentClippedDepth(Parameters.Platform);
-			}
-			else
-			{
-				return true;
-			}
-		}
-		
-		return false;
+		return TVelocityVS<PassMode>::ShouldCompilePermutation(Parameters);
 	}
 
 	static void ModifyCompilationEnvironment(const FMaterialShaderPermutationParameters& Parameters, FShaderCompilerEnvironment& OutEnvironment)
 	{
 		FVelocityPS::ModifyCompilationEnvironment(Parameters, OutEnvironment);
-		const bool bIsVelocityClippedDepthPass = (PassMode == EVelocityPassMode::Velocity_ClippedDepth);
+
+		constexpr bool bIsVelocityClippedDepthPass = (PassMode == EVelocityPassMode::Velocity_ClippedDepth);
 		OutEnvironment.SetDefine(TEXT("VELOCITY_CLIPPED_DEPTH_PASS"), bIsVelocityClippedDepthPass ? 1 : 0);
 
+		constexpr bool bIsMotionVectorPass = (PassMode == EVelocityPassMode::Velocity_StereoMotionVectors);
+		OutEnvironment.SetDefine(TEXT("STEREO_MOTION_VECTORS"), bIsMotionVectorPass ? 1 : 0);
+
 		if (bIsVelocityClippedDepthPass)
 		{
 			OutEnvironment.CompilerFlags.Add(CFLAG_AllowTypedUAVLoads);
@@ -268,11 +308,17 @@ END_GLOBAL_SHADER_PARAMETER_STRUCT()
 IMPLEMENT_SHADER_TYPE(,FVelocityVS, TEXT("/Engine/Private/VelocityShader.usf"), TEXT("MainVertexShader"), SF_Vertex);
 IMPLEMENT_SHADER_TYPE(,FVelocityPS, TEXT("/Engine/Private/VelocityShader.usf"), TEXT("MainPixelShader"), SF_Pixel);
 
+IMPLEMENT_MATERIAL_SHADER_TYPE(template<>, TVelocityVS<EVelocityPassMode::Velocity_Standard>, TEXT("/Engine/Private/VelocityShader.usf"), TEXT("MainVertexShader"), SF_Vertex);
+IMPLEMENT_MATERIAL_SHADER_TYPE(template<>, TVelocityVS<EVelocityPassMode::Velocity_ClippedDepth>, TEXT("/Engine/Private/VelocityShader.usf"), TEXT("MainVertexShader"), SF_Vertex);
+IMPLEMENT_MATERIAL_SHADER_TYPE(template<>, TVelocityVS<EVelocityPassMode::Velocity_StereoMotionVectors>, TEXT("/Engine/Private/VelocityShader.usf"), TEXT("MainVertexShader"), SF_Vertex);
+
 IMPLEMENT_MATERIAL_SHADER_TYPE(template<>, TVelocityPS<EVelocityPassMode::Velocity_Standard>, TEXT("/Engine/Private/VelocityShader.usf"), TEXT("MainPixelShader"), SF_Pixel);
 IMPLEMENT_MATERIAL_SHADER_TYPE(template<>, TVelocityPS<EVelocityPassMode::Velocity_ClippedDepth>, TEXT("/Engine/Private/VelocityShader.usf"), TEXT("MainPixelShader"), SF_Pixel);
+IMPLEMENT_MATERIAL_SHADER_TYPE(template<>, TVelocityPS<EVelocityPassMode::Velocity_StereoMotionVectors>, TEXT("/Engine/Private/VelocityShader.usf"), TEXT("MainPixelShader"), SF_Pixel);
 
-IMPLEMENT_SHADERPIPELINE_TYPE_VSPS(StandardVelocityPipeline, FVelocityVS, TVelocityPS<EVelocityPassMode::Velocity_Standard>, true);
-IMPLEMENT_SHADERPIPELINE_TYPE_VSPS(VelocityClippedDepthPipeline, FVelocityVS, TVelocityPS<EVelocityPassMode::Velocity_ClippedDepth>, true);
+IMPLEMENT_SHADERPIPELINE_TYPE_VSPS(StandardVelocityPipeline, TVelocityVS<EVelocityPassMode::Velocity_Standard>, TVelocityPS<EVelocityPassMode::Velocity_Standard>, true);
+IMPLEMENT_SHADERPIPELINE_TYPE_VSPS(VelocityClippedDepthPipeline, TVelocityVS<EVelocityPassMode::Velocity_ClippedDepth>, TVelocityPS<EVelocityPassMode::Velocity_ClippedDepth>, true);
+IMPLEMENT_SHADERPIPELINE_TYPE_VSPS(VelocityMotionVectorsPipeline, TVelocityVS<EVelocityPassMode::Velocity_StereoMotionVectors>, TVelocityPS<EVelocityPassMode::Velocity_StereoMotionVectors>, true);
 
 class FMotionVectorWorldOffsetVelocityResolveCS : public FGlobalShader
 {
@@ -523,6 +569,9 @@ void FSceneRenderer::RenderVelocities(
 	{
 		FViewInfo& View = InViews[ViewIndex];
 
+		checkf(!View.Family->EngineShowFlags.StereoMotionVectors,
+			TEXT("Normal velocity rendering is not supported alongside motion vector rendering. If this is causing problems in your project, disable r.Velocity.DirectlyRenderOpenXRMotionVectors."));
+
 		if (View.ShouldRenderView())
 		{
 			const bool bHasAnyDraw = HasAnyDraw(View.ParallelMeshDrawCommandPasses[MeshPass]);
@@ -690,6 +739,8 @@ void FSceneRenderer::RenderVelocities(
 void FMobileSceneRenderer::RenderVelocityPass(FRHICommandList& RHICmdList, const FViewInfo& View, const FInstanceCullingDrawParams* InstanceCullingDrawParams)
 {
 	checkSlow(RHICmdList.IsInsideRenderPass());
+	checkf(!View.Family->EngineShowFlags.StereoMotionVectors,
+		TEXT("Normal velocity rendering is not supported alongside motion vector rendering. If this is causing problems in your project, disable r.Velocity.DirectlyRenderOpenXRMotionVectors."));
 
 	if (auto* Pass = View.ParallelMeshDrawCommandPasses[EMeshPass::Velocity])
 	{
@@ -772,9 +823,8 @@ bool FVelocityMeshProcessor::PrimitiveHasVelocityForView(const FViewInfo& View,
 		return false;
 	}
 
-	// Only enabled when TAA is enabled on mobile
-	bool bUsesTAA = (View.AntiAliasingMethod == AAM_TemporalAA);
-	if (IsMobilePlatform(View.GetShaderPlatform()) && !bUsesTAA)
+	// Only enabled on mobile when TAA is enabled or we're rendering OpenXR motion vectors
+	if (!PlatformSupportsVelocityRendering(View.GetShaderPlatform()))
 	{
 		return false;
 	}
@@ -797,7 +847,8 @@ bool FVelocityMeshProcessor::PrimitiveHasVelocityForView(const FViewInfo& View,
 
 bool FOpaqueVelocityMeshProcessor::PrimitiveCanHaveVelocity(EShaderPlatform ShaderPlatform, const FPrimitiveSceneProxy* PrimitiveSceneProxy)
 {
-	return PrimitiveCanHaveVelocity(ShaderPlatform, PrimitiveSceneProxy->DrawsVelocity(), PrimitiveSceneProxy->HasStaticLighting());
+	const bool bDrawsVelocity = PrimitiveSceneProxy->DrawsVelocity() || VelocityIncludeStationaryPrimitives(ShaderPlatform);
+	return PrimitiveCanHaveVelocity(ShaderPlatform, bDrawsVelocity, PrimitiveSceneProxy->HasStaticLighting());
 }
 
 bool FOpaqueVelocityMeshProcessor::PrimitiveCanHaveVelocity(EShaderPlatform ShaderPlatform, bool bDrawVelocity, bool bHasStaticLighting)
@@ -815,9 +866,9 @@ bool FOpaqueVelocityMeshProcessor::PrimitiveCanHaveVelocity(EShaderPlatform Shad
 	return true;
 }
 
-bool FOpaqueVelocityMeshProcessor::PrimitiveHasVelocityForFrame(const FPrimitiveSceneProxy* PrimitiveSceneProxy)
+bool FOpaqueVelocityMeshProcessor::PrimitiveHasVelocityForFrame(EShaderPlatform ShaderPlatform, const FPrimitiveSceneProxy* PrimitiveSceneProxy)
 {
-	if (!PrimitiveSceneProxy->AlwaysHasVelocity())
+	if (!PrimitiveSceneProxy->AlwaysHasVelocity() && !VelocityIncludeStationaryPrimitives(ShaderPlatform))
 	{
 		// Check if the primitive has moved.
 		const FPrimitiveSceneInfo* PrimitiveSceneInfo = PrimitiveSceneProxy->GetPrimitiveSceneInfo();
@@ -892,7 +943,7 @@ void FOpaqueVelocityMeshProcessor::AddMeshBatch(
 
 	if (ViewIfDynamicMeshCommand)
 	{
-		if (!PrimitiveHasVelocityForFrame(PrimitiveSceneProxy))
+		if (!PrimitiveHasVelocityForFrame(ShaderPlatform, PrimitiveSceneProxy))
 		{
 			return;
 		}
@@ -927,6 +978,8 @@ void FOpaqueVelocityMeshProcessor::CollectPSOInitializers(const FSceneTexturesCo
 	const EShaderPlatform ShaderPlatform = GetFeatureLevelShaderPlatform(FeatureLevel);
 	bool bDrawsVelocity = (PreCacheParams.Mobility == EComponentMobility::Movable || PreCacheParams.Mobility == EComponentMobility::Stationary);
 	bDrawsVelocity = bDrawsVelocity || (/*VertexDeformationOutputsVelocity() &&*/ (PreCacheParams.bAnyMaterialHasWorldPositionOffset || Material.MaterialUsesWorldPositionOffset_GameThread()));
+	bDrawsVelocity = bDrawsVelocity || VelocityIncludeStationaryPrimitives(ShaderPlatform);
+
 	if (!PrimitiveCanHaveVelocity(ShaderPlatform, bDrawsVelocity, PreCacheParams.bStaticLighting))
 	{
 		return;
@@ -1089,23 +1142,35 @@ bool GetVelocityPassShaders(
 	const FMaterial& Material,
 	const FVertexFactoryType* VertexFactoryType,
 	ERHIFeatureLevel::Type FeatureLevel,
-	bool bIsTranslucentClippedDepthPass,
+	EVelocityPassMode PassMode,
 	TShaderRef<FVelocityVS>& VertexShader,
 	TShaderRef<FVelocityPS>& PixelShader)
 {
 	FMaterialShaderTypes ShaderTypes;
-	
-	ShaderTypes.AddShaderType<FVelocityVS>();
 
-	if (bIsTranslucentClippedDepthPass)
-	{
-		ShaderTypes.PipelineType = &VelocityClippedDepthPipeline;
-		ShaderTypes.AddShaderType<TVelocityPS<EVelocityPassMode::Velocity_ClippedDepth>>();
-	}
-	else
+	switch (PassMode)
 	{
-		ShaderTypes.PipelineType = &StandardVelocityPipeline;
-		ShaderTypes.AddShaderType<TVelocityPS<EVelocityPassMode::Velocity_Standard>>();
+		case EVelocityPassMode::Velocity_Standard:
+			ShaderTypes.PipelineType = &StandardVelocityPipeline;
+			ShaderTypes.AddShaderType<TVelocityVS<EVelocityPassMode::Velocity_Standard>>();
+			ShaderTypes.AddShaderType<TVelocityPS<EVelocityPassMode::Velocity_Standard>>();
+			break;
+
+		case EVelocityPassMode::Velocity_ClippedDepth:
+			ShaderTypes.PipelineType = &VelocityClippedDepthPipeline;
+			ShaderTypes.AddShaderType<TVelocityVS<EVelocityPassMode::Velocity_ClippedDepth>>();
+			ShaderTypes.AddShaderType<TVelocityPS<EVelocityPassMode::Velocity_ClippedDepth>>();
+			break;
+
+		case EVelocityPassMode::Velocity_StereoMotionVectors:
+			ShaderTypes.PipelineType = &VelocityMotionVectorsPipeline;
+			ShaderTypes.AddShaderType<TVelocityVS<EVelocityPassMode::Velocity_StereoMotionVectors>>();
+			ShaderTypes.AddShaderType<TVelocityPS<EVelocityPassMode::Velocity_StereoMotionVectors>>();
+			break;
+
+		default:
+			checkNoEntry();
+			break;
 	}
 
 	FMaterialShaders Shaders;
@@ -1135,12 +1200,24 @@ bool FVelocityMeshProcessor::Process(
 		FVelocityVS,
 		FVelocityPS> VelocityPassShaders;
 
-	bool bIsTranslucentClippedDepthPass = (MeshPassType == EMeshPass::TranslucentVelocityClippedDepth);
+	const bool bIsTranslucentClippedDepthPass = (MeshPassType == EMeshPass::TranslucentVelocityClippedDepth);
+
+	EVelocityPassMode PassMode = EVelocityPassMode::Velocity_Standard;
+	if (PlatformSupportsOpenXRMotionVectors(GetFeatureLevelShaderPlatform(FeatureLevel)))
+	{
+		checkf(!bIsTranslucentClippedDepthPass, TEXT("Translucent velocity clipped depth is not supported alongside motion vector rendering. If this is causing problems in your project, disable r.Velocity.DirectlyRenderOpenXRMotionVectors."));
+		PassMode = EVelocityPassMode::Velocity_StereoMotionVectors;
+	}
+	else if (bIsTranslucentClippedDepthPass)
+	{
+		PassMode = EVelocityPassMode::Velocity_ClippedDepth;
+	}
+
 	if (!GetVelocityPassShaders(
 		MaterialResource,
 		VertexFactory->GetType(),
 		FeatureLevel,
-		bIsTranslucentClippedDepthPass,
+		PassMode,
 		VelocityPassShaders.VertexShader,
 		VelocityPassShaders.PixelShader))
 	{
@@ -1209,13 +1286,24 @@ bool FVelocityMeshProcessor::CollectPSOInitializersInternal(
 		FVelocityVS,
 		FVelocityPS> VelocityPassShaders;
 
-	bool bIsTranslucentClippedDepthPass = (MeshPassType == EMeshPass::TranslucentVelocityClippedDepth);
+	const bool bIsTranslucentClippedDepthPass = (MeshPassType == EMeshPass::TranslucentVelocityClippedDepth);
+
+	EVelocityPassMode PassMode = EVelocityPassMode::Velocity_Standard;
+	if (PlatformSupportsOpenXRMotionVectors(GetFeatureLevelShaderPlatform(FeatureLevel)))
+	{
+		checkf(!bIsTranslucentClippedDepthPass, TEXT("Translucent velocity clipped depth is not supported alongside motion vector rendering. If this is causing problems in your project, disable r.Velocity.DirectlyRenderOpenXRMotionVectors."));
+		PassMode = EVelocityPassMode::Velocity_StereoMotionVectors;
+	}
+	else if (bIsTranslucentClippedDepthPass)
+	{
+		PassMode = EVelocityPassMode::Velocity_ClippedDepth;
+	}
 
 	if (!GetVelocityPassShaders(
 		MaterialResource,
 		VertexFactoryData.VertexFactoryType,
 		FeatureLevel,
-		bIsTranslucentClippedDepthPass,
+		PassMode,
 		VelocityPassShaders.VertexShader,
 		VelocityPassShaders.PixelShader))
 	{
@@ -1270,7 +1358,9 @@ FMeshPassProcessor* CreateVelocityPassProcessor(ERHIFeatureLevel::Type FeatureLe
 	
 	FMeshPassProcessorRenderState VelocityPassState;
 	VelocityPassState.SetBlendState(TStaticBlendState<CW_RGBA>::GetRHI());
-	VelocityPassState.SetDepthStencilState((EarlyZPassMode == DDM_AllOpaqueNoVelocity) // if the depth mode is all opaque except velocity, it relies on velocity to write the depth of the remaining meshes
+
+	const bool bNeedStationaryPrimitiveDepth = VelocityIncludeStationaryPrimitives(GetFeatureLevelShaderPlatform(FeatureLevel));
+	VelocityPassState.SetDepthStencilState(bNeedStationaryPrimitiveDepth || (EarlyZPassMode == DDM_AllOpaqueNoVelocity) // if the depth mode is all opaque except velocity, it relies on velocity to write the depth of the remaining meshes
 										    ? TStaticDepthStencilState<true, CF_DepthNearOrEqual>::GetRHI()
 											: TStaticDepthStencilState<false, CF_DepthNearOrEqual>::GetRHI());
 
diff --git a/Engine/Source/Runtime/Renderer/Private/VelocityRendering.h b/Engine/Source/Runtime/Renderer/Private/VelocityRendering.h
--- a/Engine/Source/Runtime/Renderer/Private/VelocityRendering.h
+++ b/Engine/Source/Runtime/Renderer/Private/VelocityRendering.h
@@ -119,7 +119,7 @@ public:
 	static bool PrimitiveCanHaveVelocity(EShaderPlatform ShaderPlatform, bool bDrawVelocity, bool bHasStaticLighting);
 
 	/** Returns true if the primitive has velocity for the current frame. */
-	static bool PrimitiveHasVelocityForFrame(const FPrimitiveSceneProxy* PrimitiveSceneProxy);
+	static bool PrimitiveHasVelocityForFrame(EShaderPlatform ShaderPlatform, const FPrimitiveSceneProxy* PrimitiveSceneProxy);
 
 private:
 	bool TryAddMeshBatch(
