OpenXR now building for all platforms.
-The platform specificity of the openxr plugins and modules was problematic because content derived from these would not load on non-xr platforms for projects that are meant to be cross platform.  OpenXR is not actually platform dependent.  It can be implemented on any platform.  If it is present as a fully independent openxr runtime a small amount of platform specific code is needed to implement an openxr loader.  It is also possible to create an unreal plugin which either acts as an openxr runtime or is able to find an openxr runtime.
-The plugins are still disabled by default, so only projects that enable them will contain the modules and conent.

diff --git a/Engine/Platforms/VisionOS/Plugins/Runtime/OpenXR/OpenXR_VisionOS.uplugin b/Engine/Platforms/VisionOS/Plugins/Runtime/OpenXR/OpenXR_VisionOS.uplugin
deleted file mode 100644
diff --git a/Engine/Platforms/VisionOS/Plugins/Runtime/OpenXRHandTracking/OpenXRHandTracking_VisionOS.uplugin b/Engine/Platforms/VisionOS/Plugins/Runtime/OpenXRHandTracking/OpenXRHandTracking_VisionOS.uplugin
deleted file mode 100644
diff --git a/Engine/Plugins/Runtime/OpenXR/OpenXR.uplugin b/Engine/Plugins/Runtime/OpenXR/OpenXR.uplugin
--- a/Engine/Plugins/Runtime/OpenXR/OpenXR.uplugin
+++ b/Engine/Plugins/Runtime/OpenXR/OpenXR.uplugin
@@ -14,58 +14,26 @@
 	"CanContainContent": true,
 	"IsBetaVersion": false,
 	"Installed": false,
-	"SupportedTargetPlatforms": [
-		"Win64",
-		"Linux",
-		"LinuxArm64",
-		"Android"
-	],
 	"Modules": [
 		{
 			"Name": "OpenXRHMD",
 			"Type": "Runtime",
-			"LoadingPhase": "PostConfigInit",
-			"PlatformAllowList": [
-				"Win64",
-				"Linux",
-				"LinuxArm64",
-				"Android",
-				"Mac"
-			]
+			"LoadingPhase": "PostConfigInit"
 		},
 		{
 			"Name": "OpenXRInput",
 			"Type": "Runtime",
-			"LoadingPhase": "PostEngineInit",
-			"PlatformAllowList": [
-				"Win64",
-				"Linux",
-				"LinuxArm64",
-				"Android",
-				"Mac"
-			]
+			"LoadingPhase": "PostEngineInit"
 		},
 		{
 			"Name": "OpenXRAR",
 			"Type": "Runtime",
-			"LoadingPhase": "PostConfigInit",
-			"PlatformAllowList": [
-				"Win64",
-				"Linux",
-				"LinuxArm64",
-				"Android",
-				"Mac"
-			]
+			"LoadingPhase": "PostConfigInit"
 		},
 		{
 			"Name": "OpenXREditor",
 			"Type": "Editor",
-			"LoadingPhase": "Default",
-			"PlatformAllowList": [
-				"Win64",
-				"Linux",
-				"Mac"
-			]
+			"LoadingPhase": "Default"
 		}
 	],
 	"Plugins": [
@@ -78,4 +46,4 @@
 			"Enabled": true
 		}
 	]
-}
\ No newline at end of file
+}
diff --git a/Engine/Plugins/Runtime/OpenXREyeTracker/OpenXREyeTracker.uplugin b/Engine/Plugins/Runtime/OpenXREyeTracker/OpenXREyeTracker.uplugin
--- a/Engine/Plugins/Runtime/OpenXREyeTracker/OpenXREyeTracker.uplugin
+++ b/Engine/Plugins/Runtime/OpenXREyeTracker/OpenXREyeTracker.uplugin
@@ -14,24 +14,11 @@
 	"CanContainContent": true,
 	"IsBetaVersion": false,
 	"Installed": false,
-	"SupportedTargetPlatforms": [
-		"Win64",
-		"Linux",
-		"LinuxArm64",
-		"Android"
-	],
 	"Modules": [
 		{
 			"Name": "OpenXREyeTracker",
 			"Type": "Runtime",
-			"LoadingPhase": "PostConfigInit",
-			"PlatformAllowList": [
-				"Win64",
-				"Linux",
-				"LinuxArm64",
-				"Android",
-				"Mac"
-			]
+			"LoadingPhase": "PostConfigInit"
 		}
 	],
 	"Plugins": [
@@ -44,4 +31,4 @@
 			"Enabled": true
 		}
 	]
-}
\ No newline at end of file
+}
diff --git a/Engine/Plugins/Runtime/OpenXRHandTracking/OpenXRHandTracking.uplugin b/Engine/Plugins/Runtime/OpenXRHandTracking/OpenXRHandTracking.uplugin
--- a/Engine/Plugins/Runtime/OpenXRHandTracking/OpenXRHandTracking.uplugin
+++ b/Engine/Plugins/Runtime/OpenXRHandTracking/OpenXRHandTracking.uplugin
@@ -14,47 +14,22 @@
 	"CanContainContent": true,
 	"IsBetaVersion": false,
 	"Installed": false,
-	"SupportedTargetPlatforms": [
-		"Win64",
-		"Linux",
-		"LinuxArm64",
-		"Android"
-	],
 	"Modules": [
 		{
 			"Name": "OpenXRHandTracking",
 			"Type": "Runtime",
-			"LoadingPhase": "PostConfigInit",
-			"PlatformAllowList": [
-				"Win64",
-				"Linux",
-				"LinuxArm64",
-				"Android",
-				"Mac"
-			]
+			"LoadingPhase": "PostConfigInit"
 		},
 		{
 			"Name": "OpenXRHandTrackingEditor",
 			"Type": "Editor",
-			"LoadingPhase": "Default",
-			"PlatformAllowList": [
-				"Win64",
-				"Linux",
-				"Mac"
-			]
+			"LoadingPhase": "Default"
 		}
 	],
 	"Plugins": [
 		{
 			"Name": "OpenXR",
-			"Enabled": true,
-			"PlatformAllowList": [
-				"Win64",
-				"Linux",
-				"LinuxArm64",
-				"Android",
-				"Mac"
-			]
+			"Enabled": true
 		},
 		{
 			"Name": "LiveLink",
diff --git a/Templates/TP_VirtualRealityBP/TP_VirtualRealityBP.uproject b/Templates/TP_VirtualRealityBP/TP_VirtualRealityBP.uproject
--- a/Templates/TP_VirtualRealityBP/TP_VirtualRealityBP.uproject
+++ b/Templates/TP_VirtualRealityBP/TP_VirtualRealityBP.uproject
@@ -6,35 +6,15 @@
 	"Plugins": [
 		{
 			"Name": "OpenXR",
-			"Enabled": true,
-			"SupportedTargetPlatforms": [
-				"Win64",
-				"Linux",
-				"LinuxArm64",
-				"Android",
-				"PS5"
-			]
+			"Enabled": true
 		},
 		{
 			"Name": "OpenXREyeTracker",
-			"Enabled": true,
-			"SupportedTargetPlatforms": [
-				"Win64",
-				"Linux",
-				"LinuxArm64",
-				"Android"
-			]
+			"Enabled": true
 		},
 		{
 			"Name": "OpenXRHandTracking",
-			"Enabled": true,
-			"SupportedTargetPlatforms": [
-				"Win64",
-				"Linux",
-				"LinuxArm64",
-				"Android",
-				"PS5"
-			]
+			"Enabled": true
 		}
 	],
 	"TargetPlatforms": [
