OpenXRHMD dx11 build break fix
-Improving the method for detecting which graphics api's are supported.
#rnx

diff --git a/Engine/Plugins/Runtime/OpenXR/Source/OpenXRHMD/OpenXRHMD.Build.cs b/Engine/Plugins/Runtime/OpenXR/Source/OpenXRHMD/OpenXRHMD.Build.cs
--- a/Engine/Plugins/Runtime/OpenXR/Source/OpenXRHMD/OpenXRHMD.Build.cs
+++ b/Engine/Plugins/Runtime/OpenXR/Source/OpenXRHMD/OpenXRHMD.Build.cs
@@ -49,6 +49,9 @@ namespace UnrealBuildTool.Rules
 
             if (Target.Platform == UnrealTargetPlatform.Win64)
             {
+				PrivateDefinitions.Add("OPENXRHMD_HAVE_DX11");
+				PrivateDefinitions.Add("OPENXRHMD_HAVE_DX12");
+
 				PublicDependencyModuleNames.AddRange(new string[] {
 					"D3D11RHI",
 					"D3D12RHI"
@@ -65,6 +68,16 @@ namespace UnrealBuildTool.Rules
 
             if (Target.Platform == UnrealTargetPlatform.Win64 || Target.Platform == UnrealTargetPlatform.Android)
             {
+				if (Target.Platform == UnrealTargetPlatform.Win64)
+				{
+					PrivateDefinitions.Add("OPENXRHMD_HAVE_OPENGL");
+				}
+			
+				if (Target.Platform == UnrealTargetPlatform.Android)
+				{
+					PrivateDefinitions.Add("OPENXRHMD_HAVE_OPENGLES");
+				}
+
 				PublicDependencyModuleNames.Add("OpenGLDrv");
 
 				if (!bUsePrecompiled || Target.LinkType == TargetLinkType.Monolithic)
@@ -76,6 +89,8 @@ namespace UnrealBuildTool.Rules
 			if (Target.Platform == UnrealTargetPlatform.Win64 || Target.Platform == UnrealTargetPlatform.Android  
 			    || Target.IsInPlatformGroup(UnrealPlatformGroup.Linux))
             {
+				PrivateDefinitions.Add("OPENXRHMD_HAVE_VULKAN");
+
 				PublicDependencyModuleNames.Add("VulkanRHI");
 
 				if (!bUsePrecompiled || Target.LinkType == TargetLinkType.Monolithic)
diff --git a/Engine/Plugins/Runtime/OpenXR/Source/OpenXRHMD/Public/OpenXRPlatformRHI.h b/Engine/Plugins/Runtime/OpenXR/Source/OpenXRHMD/Public/OpenXRPlatformRHI.h
--- a/Engine/Plugins/Runtime/OpenXR/Source/OpenXRHMD/Public/OpenXRPlatformRHI.h
+++ b/Engine/Plugins/Runtime/OpenXR/Source/OpenXRHMD/Public/OpenXRPlatformRHI.h
@@ -6,20 +6,29 @@
 
 #if PLATFORM_WINDOWS
 #define XR_USE_PLATFORM_WIN32		1
+#endif
+
+#ifdef OPENXRHMD_HAVE_DX11
 #define XR_USE_GRAPHICS_API_D3D11	1
+#endif
+
+#ifdef OPENXRHMD_HAVE_DX12
 #define XR_USE_GRAPHICS_API_D3D12	1
 #endif
 
-#if PLATFORM_WINDOWS
+#ifdef OPENXRHMD_HAVE_OPENGL
 #define XR_USE_GRAPHICS_API_OPENGL	1
 #endif
 
 #if PLATFORM_ANDROID
 #define XR_USE_PLATFORM_ANDROID 1
+#endif
+
+#ifdef OPENXRHMD_HAVE_OPENGLES
 #define XR_USE_GRAPHICS_API_OPENGL_ES	1
 #endif
 
-#if PLATFORM_WINDOWS || PLATFORM_ANDROID || PLATFORM_LINUX
+#ifdef OPENXRHMD_HAVE_VULKAN
 #define XR_USE_GRAPHICS_API_VULKAN 1
 #endif
 
diff --git a/Engine/Plugins/Runtime/XR/XRScribe/Source/XRScribe/XRScribe.Build.cs b/Engine/Plugins/Runtime/XR/XRScribe/Source/XRScribe/XRScribe.Build.cs
--- a/Engine/Plugins/Runtime/XR/XRScribe/Source/XRScribe/XRScribe.Build.cs
+++ b/Engine/Plugins/Runtime/XR/XRScribe/Source/XRScribe/XRScribe.Build.cs
@@ -8,7 +8,7 @@ namespace UnrealBuildTool.Rules
 	// TODO Possibly split into capture module and replay module??
 	public class XRScribe : ModuleRules
 	{
-        public XRScribe(ReadOnlyTargetRules Target) : base(Target)
+		public XRScribe(ReadOnlyTargetRules Target) : base(Target)
 		{
 			PublicDependencyModuleNames.AddRange(new string[]
 			{
@@ -22,9 +22,15 @@ namespace UnrealBuildTool.Rules
 
 			PrivateDependencyModuleNames.AddRange(new string[]
 			{
-                "DeveloperSettings",
+				"DeveloperSettings",
 				"Projects",
 			});
+
+			if (Target.Platform == UnrealTargetPlatform.Win64)
+			{
+				PrivateDefinitions.Add("OPENXRHMD_HAVE_DX11");
+				PrivateDefinitions.Add("OPENXRHMD_HAVE_DX12");
+			}
 		}
 	}
 }
