Add a new CVar, xr.PreferFBSpaceWarp, to control whether XR_FB_space_warp is used over XR_EXT_frame_synthesis when both extensions are available.
The CVar defaults to false, but is set to true in the Quest 3 device profile.

diff --git a/Engine/Config/BaseDeviceProfiles.ini b/Engine/Config/BaseDeviceProfiles.ini
--- a/Engine/Config/BaseDeviceProfiles.ini
+++ b/Engine/Config/BaseDeviceProfiles.ini
@@ -1282,6 +1282,7 @@ DeviceType=Android
 BaseProfileName=Android_OpenXR
 +CVars=r.Mobile.Oculus.ForceSymmetric=1
 +CVars=r.Mobile.AdrenoOcclusionMode=1
++CVars=xr.PreferFBSpaceWarp=1
 
 [Meta_Quest_Pro DeviceProfile]
 DeviceType=Android
diff --git a/Engine/Plugins/Runtime/OpenXR/Source/OpenXRHMD/Private/OpenXRHMDModule.cpp b/Engine/Plugins/Runtime/OpenXR/Source/OpenXRHMD/Private/OpenXRHMDModule.cpp
--- a/Engine/Plugins/Runtime/OpenXR/Source/OpenXRHMD/Private/OpenXRHMDModule.cpp
+++ b/Engine/Plugins/Runtime/OpenXR/Source/OpenXRHMD/Private/OpenXRHMDModule.cpp
@@ -45,6 +45,12 @@ static TAutoConsoleVariable<bool> CVarRetainPreInitInstance(
 	TEXT("If true, OpenXR will retain any instance created during PreInit rather than destroying it.  Destroying it is more correct because we are not yet certain to have chosen OpenXRHMD, and another HMD plugin could take over and try to create an instance of its own which would fail on a runtime that supports only one."),
 	ECVF_RenderThreadSafe);
 
+static TAutoConsoleVariable<bool> CVarPreferFBSpaceWarp(
+	TEXT("xr.PreferFBSpaceWarp"),
+	false,
+	TEXT("If true, OpenXR will use XR_FB_space_warp rather than XR_EXT_frame_synthesis to provide frame synthesis when both extensions are available."),
+	ECVF_ReadOnly);
+
 //---------------------------------------------------
 // OpenXRHMD Plugin Implementation
 //---------------------------------------------------
@@ -588,8 +594,8 @@ bool FOpenXRHMDModule::GetOptionalExtensions(TArray<const ANSICHAR*>& OutExtensi
 	OutExtensions.Add(XR_OCULUS_AUDIO_DEVICE_GUID_EXTENSION_NAME);
 #endif
 
-	// These extensions are mutually incompatible because XR_FB_COMPOSITION_LAYER_ALPHA_BLEND disables non-opaque environment blend modes
-	// XR_EXT_COMPOSITION_LAYER_INVERTED_ALPHA is the more modern and platform-agnostic option and is preferred when available
+	// These extensions are mutually incompatible because XR_FB_composition_layer_alpha_blend disables non-opaque environment blend modes
+	// XR_EXT_composition_layer_inverted_alpha is the more modern and platform-agnostic option and is preferred when available
 	if (AvailableExtensions.Contains(XR_EXT_COMPOSITION_LAYER_INVERTED_ALPHA_EXTENSION_NAME))
 	{
 		OutExtensions.Add(XR_EXT_COMPOSITION_LAYER_INVERTED_ALPHA_EXTENSION_NAME);
@@ -600,14 +606,32 @@ bool FOpenXRHMDModule::GetOptionalExtensions(TArray<const ANSICHAR*>& OutExtensi
 	}
 
 	// These extensions both perform frame synthesis, and require providing motion vector and motion vector depth swapchains at a requested size
-	// XR_EXT_FRAME_SYNTHESIS is not currently functioning properly on Quest, so XR_FB_SPACE_WARP is preferred when available
-	if (AvailableExtensions.Contains(XR_FB_SPACE_WARP_EXTENSION_NAME))
+	// XR_EXT_frame_synthesis is the more modern and platform-agnostic option and is preferred when available
+	// However, it is not currently working properly on Oculus runtimes, so XR_FB_space_warp is preferred on Quest devices
+	const bool bSpaceWarpSupported = AvailableExtensions.Contains(XR_FB_SPACE_WARP_EXTENSION_NAME);
+	const bool bFrameSynthesisSupported = AvailableExtensions.Contains(XR_EXT_FRAME_SYNTHESIS_EXTENSION_NAME);
+
+	if (bSpaceWarpSupported && bFrameSynthesisSupported)
 	{
-		OutExtensions.Add(XR_FB_SPACE_WARP_EXTENSION_NAME);
+		if (CVarPreferFBSpaceWarp.GetValueOnAnyThread())
+		{
+			OutExtensions.Add(XR_FB_SPACE_WARP_EXTENSION_NAME);
+		}
+		else
+		{
+			OutExtensions.Add(XR_EXT_FRAME_SYNTHESIS_EXTENSION_NAME);
+		}
 	}
-	else if (AvailableExtensions.Contains(XR_EXT_FRAME_SYNTHESIS_EXTENSION_NAME))
+	else
 	{
-		OutExtensions.Add(XR_EXT_FRAME_SYNTHESIS_EXTENSION_NAME);
+		if (bFrameSynthesisSupported)
+		{
+			OutExtensions.Add(XR_EXT_FRAME_SYNTHESIS_EXTENSION_NAME);
+		}
+		else if (bSpaceWarpSupported)
+		{
+			OutExtensions.Add(XR_FB_SPACE_WARP_EXTENSION_NAME);
+		}
 	}
 	
 	OutExtensions.Add(XR_EXT_LOCAL_FLOOR_EXTENSION_NAME);
