Update OpenXRHMD to use XR_EXT_frame_synthesis on Quest, rather than XR_FB_space_warp.
- Add the XR_FRAME_SYNTHESIS_INFO_REQUEST_RELAXED_FRAME_INTERVAL_BIT_EXT flag needed to enable frame synthesis to the frame synthesis layer info struct
- Remove xr.PreferFBSpaceWarp=1 from the Meta_Quest_3 device profile. This was not functioning correctly in the first place because we call FOpenXRHMDModule::GetOptionalExtensions(..) before device profiles are loaded

diff --git a/Engine/Config/BaseDeviceProfiles.ini b/Engine/Config/BaseDeviceProfiles.ini
--- a/Engine/Config/BaseDeviceProfiles.ini
+++ b/Engine/Config/BaseDeviceProfiles.ini
@@ -1306,7 +1306,6 @@ DeviceType=Android
 BaseProfileName=Android_OpenXR
 +CVars=r.Mobile.Oculus.ForceSymmetric=1
 +CVars=r.Mobile.AdrenoOcclusionMode=1
-+CVars=xr.PreferFBSpaceWarp=1
 
 [Meta_Quest_Pro DeviceProfile]
 DeviceType=Android
diff --git a/Engine/Plugins/Runtime/OpenXR/Source/OpenXRHMD/Private/OpenXRHMD.cpp b/Engine/Plugins/Runtime/OpenXR/Source/OpenXRHMD/Private/OpenXRHMD.cpp
--- a/Engine/Plugins/Runtime/OpenXR/Source/OpenXRHMD/Private/OpenXRHMD.cpp
+++ b/Engine/Plugins/Runtime/OpenXR/Source/OpenXRHMD/Private/OpenXRHMD.cpp
@@ -1527,7 +1527,7 @@ void FOpenXRHMD::PostRenderViewFamily_RenderThread(FRDGBuilder& GraphBuilder, FS
 
 				FrameSynthesisLayerInfo.type = XR_TYPE_FRAME_SYNTHESIS_INFO_EXT;
 				FrameSynthesisLayerInfo.next = nullptr;
-				FrameSynthesisLayerInfo.layerFlags = 0;
+				FrameSynthesisLayerInfo.layerFlags = XR_FRAME_SYNTHESIS_INFO_REQUEST_RELAXED_FRAME_INTERVAL_BIT_EXT;
 				FrameSynthesisLayerInfo.motionVectorSubImage = MotionVectorImage;
 				FrameSynthesisLayerInfo.motionVectorScale = { 1.0f, 1.0f, 1.0f, 0.0f };
 				FrameSynthesisLayerInfo.motionVectorOffset = {};
