Fix FBFoveation race condition during startup (could cause crashes repro-able with -onethread), and make FBFoveation CVars read-only to avoid crashes when toggling at runtime.
These CVars can be switched back to runtime-toggleable after OpenXRHMD RDG synchronization issues are addressed.

#Synced-CL 48355415

diff --git a/Engine/Plugins/Runtime/OpenXR/Source/OpenXRHMD/Private/FBFoveationImageGenerator.cpp b/Engine/Plugins/Runtime/OpenXR/Source/OpenXRHMD/Private/FBFoveationImageGenerator.cpp
--- a/Engine/Plugins/Runtime/OpenXR/Source/OpenXRHMD/Private/FBFoveationImageGenerator.cpp
+++ b/Engine/Plugins/Runtime/OpenXR/Source/OpenXRHMD/Private/FBFoveationImageGenerator.cpp
@@ -11,19 +11,19 @@ static TAutoConsoleVariable<int32> CVarOpenXRFBFoveationLevel(
 	0,
 	TEXT("Possible foveation levels as specified by the XrFoveationLevelFB enumeration.\n")
 	TEXT("0 = None, 1 = Low , 2 = Medium, 3 = High.\n"),
-	ECVF_Default);
+	ECVF_RenderThreadSafe | ECVF_ReadOnly);
 
 static TAutoConsoleVariable<bool> CVarOpenXRFBFoveationDynamic(
 	TEXT("xr.OpenXRFBFoveationDynamic"),
 	false,
 	TEXT("Whether dynamically changing foveation based on performance headroom is enabled.\n"),
-	ECVF_Default);
+	ECVF_RenderThreadSafe | ECVF_ReadOnly);
 
 static TAutoConsoleVariable<float> CVarOpenXRFBFoveationVerticalOffset(
 	TEXT("xr.OpenXRFBFoveationVerticalOffset"),
 	0,
 	TEXT("Desired vertical offset in degrees for the center of the foveation pattern.\n"),
-	ECVF_Default);
+	ECVF_RenderThreadSafe | ECVF_ReadOnly);
 
 static TAutoConsoleVariable<int32> CVarFBFoveationPreview(
 	TEXT("xr.OpenXRFBFoveation.Preview"),
