Reducing FRHICommandListImmediate usage
- IStereoRenderTargetManager::AllocateRenderTargetTextures should take a command list.
- Removing deprecated methods from IStereoRenderTargetManager.

diff --git a/Engine/Plugins/Runtime/OpenXR/Source/OpenXRHMD/Private/OpenXRHMD.cpp b/Engine/Plugins/Runtime/OpenXR/Source/OpenXRHMD/Private/OpenXRHMD.cpp
--- a/Engine/Plugins/Runtime/OpenXR/Source/OpenXRHMD/Private/OpenXRHMD.cpp
+++ b/Engine/Plugins/Runtime/OpenXR/Source/OpenXRHMD/Private/OpenXRHMD.cpp
@@ -2777,7 +2777,7 @@ int32 FOpenXRHMD::AcquireDepthTexture()
 	return 0;
 }
 
-bool FOpenXRHMD::AllocateRenderTargetTextures(uint32 SizeX, uint32 SizeY, uint8 Format, uint32 NumMips, ETextureCreateFlags Flags, ETextureCreateFlags TargetableTextureFlags, TArray<FTextureRHIRef>& OutTargetableTextures, TArray<FTextureRHIRef>& OutShaderResourceTextures, uint32 NumSamples)
+bool FOpenXRHMD::AllocateRenderTargetTextures(FRHICommandListBase& RHICmdList, uint32 SizeX, uint32 SizeY, uint8 Format, uint32 NumMips, ETextureCreateFlags Flags, ETextureCreateFlags TargetableTextureFlags, TArray<FTextureRHIRef>& OutTargetableTextures, TArray<FTextureRHIRef>& OutShaderResourceTextures, uint32 NumSamples)
 {
 	check(IsInRenderingThread());
 
@@ -2847,15 +2847,14 @@ bool FOpenXRHMD::AllocateRenderTargetTextures(uint32 SizeX, uint32 SizeY, uint8
 	// Image will be acquired by the renderer if supported, if not we acquire it ahead of time here
 	if (!bIsAcquireOnAnyThreadSupported)
 	{
-		FRHICommandListImmediate& RHICmdList = FRHICommandListImmediate::Get();
-		RHICmdList.EnqueueLambda([this, Swapchain](FRHICommandListImmediate&)
+		RHICmdList.EnqueueLambda([this, Swapchain](FRHICommandListBase&)
 		{
 			RenderBridge->RunOnRHISubmissionThread([Swapchain]()
 			{
 				Swapchain->IncrementSwapChainIndex_RHIThread();
 			});
 		});
-		RHICmdList.ImmediateFlush(EImmediateFlushType::FlushRHIThread);
+		RHICmdList.GetAsImmediate().ImmediateFlush(EImmediateFlushType::FlushRHIThread);
 	}
 	if (Swapchain)
 	{
@@ -2898,16 +2897,15 @@ bool FOpenXRHMD::AllocateRenderTargetTextures(uint32 SizeX, uint32 SizeY, uint8
 		// Image will be acquired by the renderer if supported, if not we acquire it ahead of time here
 		if (!bIsAcquireOnAnyThreadSupported)
 		{
-			FRHICommandListImmediate& RHICmdList = FRHICommandListImmediate::Get();
 			FXRSwapChainPtr& EmulationSwapchain = PipelinedLayerStateRendering.EmulatedLayerState.EmulationSwapchain;
-			RHICmdList.EnqueueLambda([this, EmulationSwapchain](FRHICommandListImmediate&)
+				RHICmdList.EnqueueLambda([this, EmulationSwapchain](FRHICommandListBase&)
 			{
 				RenderBridge->RunOnRHISubmissionThread([EmulationSwapchain]
 				{
 					EmulationSwapchain->IncrementSwapChainIndex_RHIThread();
 				});
 			});
-			RHICmdList.ImmediateFlush(EImmediateFlushType::FlushRHIThread);
+				RHICmdList.GetAsImmediate().ImmediateFlush(EImmediateFlushType::FlushRHIThread);
 		}
 	}
 
@@ -2949,10 +2947,9 @@ bool FOpenXRHMD::AllocateRenderTargetTextures(uint32 SizeX, uint32 SizeY, uint8
 		// Images will be acquired by the renderer if supported, if not we acquire it ahead of time here
 		if (!bIsAcquireOnAnyThreadSupported)
 		{
-			FRHICommandListImmediate& RHICmdList = FRHICommandListImmediate::Get();
 			FXRSwapChainPtr& MotionVectorSwapchain = PipelinedLayerStateRendering.MotionVectorSwapchain;
 			FXRSwapChainPtr& MotionVectorDepthSwapchain = PipelinedLayerStateRendering.MotionVectorDepthSwapchain;
-			RHICmdList.EnqueueLambda([this, MotionVectorSwapchain, MotionVectorDepthSwapchain](FRHICommandListImmediate&)
+			RHICmdList.EnqueueLambda([this, MotionVectorSwapchain, MotionVectorDepthSwapchain](FRHICommandListBase&)
 			{
 				RenderBridge->RunOnRHISubmissionThread([MotionVectorSwapchain, MotionVectorDepthSwapchain]
 				{
@@ -2960,7 +2957,7 @@ bool FOpenXRHMD::AllocateRenderTargetTextures(uint32 SizeX, uint32 SizeY, uint8
 					MotionVectorDepthSwapchain->IncrementSwapChainIndex_RHIThread();
 				});
 			});
-			RHICmdList.ImmediateFlush(EImmediateFlushType::FlushRHIThread);
+			RHICmdList.GetAsImmediate().ImmediateFlush(EImmediateFlushType::FlushRHIThread);
 		}
 	}
 
diff --git a/Engine/Plugins/Runtime/OpenXR/Source/OpenXRHMD/Private/OpenXRHMD.h b/Engine/Plugins/Runtime/OpenXR/Source/OpenXRHMD/Private/OpenXRHMD.h
--- a/Engine/Plugins/Runtime/OpenXR/Source/OpenXRHMD/Private/OpenXRHMD.h
+++ b/Engine/Plugins/Runtime/OpenXR/Source/OpenXRHMD/Private/OpenXRHMD.h
@@ -433,7 +433,7 @@ public:
 	/** IStereoRenderTargetManager */
 	virtual bool ShouldUseSeparateRenderTarget() const override { return IsStereoEnabled() && RenderBridge.IsValid(); }
 	virtual void CalculateRenderTargetSize(const FViewport& Viewport, uint32& InOutSizeX, uint32& InOutSizeY) override;
-	virtual bool AllocateRenderTargetTextures(uint32 SizeX, uint32 SizeY, uint8 Format, uint32 NumLayers, ETextureCreateFlags Flags, ETextureCreateFlags TargetableTextureFlags, TArray<FTextureRHIRef>& OutTargetableTextures, TArray<FTextureRHIRef>& OutShaderResourceTextures, uint32 NumSamples = 1) override;
+	virtual bool AllocateRenderTargetTextures(FRHICommandListBase& RHICmdList, uint32 SizeX, uint32 SizeY, uint8 Format, uint32 NumLayers, ETextureCreateFlags Flags, ETextureCreateFlags TargetableTextureFlags, TArray<FTextureRHIRef>& OutTargetableTextures, TArray<FTextureRHIRef>& OutShaderResourceTextures, uint32 NumSamples = 1) override;
 	virtual int32 AcquireColorTexture() override final;
 	virtual int32 AcquireDepthTexture() override final;
 	virtual bool AllocateDepthTexture(uint32 Index, uint32 SizeX, uint32 SizeY, uint8 Format, uint32 NumMips, ETextureCreateFlags InTexFlags, ETextureCreateFlags TargetableTextureFlags, FTextureRHIRef& OutTargetableTexture, FTextureRHIRef& OutShaderResourceTexture, uint32 NumSamples = 1) override final;
diff --git a/Engine/Source/Runtime/Engine/Private/Slate/SceneViewport.cpp b/Engine/Source/Runtime/Engine/Private/Slate/SceneViewport.cpp
--- a/Engine/Source/Runtime/Engine/Private/Slate/SceneViewport.cpp
+++ b/Engine/Source/Runtime/Engine/Private/Slate/SceneViewport.cpp
@@ -2055,7 +2055,7 @@ void FSceneViewport::InitRHI(FRHICommandListBase& RHICmdList)
 	if(bRequiresHitProxyStorage)
 	{
 		// Initialize the hit proxy map.
-		HitProxyMap.Init(SizeX,SizeY);
+		HitProxyMap.Init(RHICmdList, SizeX,SizeY);
 	}
 	RTTSize = FIntPoint(0, 0);
 
@@ -2082,10 +2082,7 @@ void FSceneViewport::InitRHI(FRHICommandListBase& RHICmdList)
 		if (StereoRenderTargetManager != nullptr)
 		{
 			StereoRenderTargetManager->CalculateRenderTargetSize(*this, TexSizeX, TexSizeY);
-			PRAGMA_DISABLE_DEPRECATION_WARNINGS
-			NumBufferedFrames = StereoRenderTargetManager->GetNumberOfBufferedFrames();
-			PRAGMA_ENABLE_DEPRECATION_WARNINGS
-			bHMDAllocatedSeparateRenderTargets = StereoRenderTargetManager->AllocateRenderTargetTextures(TexSizeX, TexSizeY, SceneTargetFormat, 1, TexCreate_None, TexCreate_RenderTargetable, BufferedRTRHI, BufferedSRVRHI);
+			bHMDAllocatedSeparateRenderTargets = StereoRenderTargetManager->AllocateRenderTargetTextures(RHICmdList, TexSizeX, TexSizeY, SceneTargetFormat, 1, TexCreate_None, TexCreate_RenderTargetable, BufferedRTRHI, BufferedSRVRHI);
 			if (bHMDAllocatedSeparateRenderTargets)
 			{
 				check(BufferedRTRHI.Num() == BufferedSRVRHI.Num());
@@ -2135,9 +2132,7 @@ void FSceneViewport::InitRHI(FRHICommandListBase& RHICmdList)
 				RTRHI = BufferedRTRHI[i];
 				SRVRHI = BufferedSRVRHI[i];
 			}
-			PRAGMA_DISABLE_DEPRECATION_WARNINGS
-			else if (StereoRenderTargetManager == nullptr || !StereoRenderTargetManager->AllocateRenderTargetTexture(i, TexSizeX, TexSizeY, SceneTargetFormat, 1, TexCreate_None, TexCreate_RenderTargetable, RTRHI, SRVRHI))
-			PRAGMA_ENABLE_DEPRECATION_WARNINGS
+			else if (StereoRenderTargetManager == nullptr)
 			{
 				const FRHITextureCreateDesc Desc =
 					FRHITextureCreateDesc::Create2D(TEXT("BufferedRT"))
diff --git a/Engine/Source/Runtime/Engine/Private/UnrealClient.cpp b/Engine/Source/Runtime/Engine/Private/UnrealClient.cpp
--- a/Engine/Source/Runtime/Engine/Private/UnrealClient.cpp
+++ b/Engine/Source/Runtime/Engine/Private/UnrealClient.cpp
@@ -2246,7 +2246,7 @@ void FViewport::InitRHI(FRHICommandListBase& RHICmdList)
 	if(bRequiresHitProxyStorage)
 	{
 		// Initialize the hit proxy map.
-		HitProxyMap.Init(SizeX,SizeY);
+		HitProxyMap.Init(RHICmdList, SizeX, SizeY);
 	}
 }
 
@@ -2267,7 +2267,7 @@ FViewport::FHitProxyMap::~FHitProxyMap()
 }
 
 
-void FViewport::FHitProxyMap::Init(uint32 NewSizeX,uint32 NewSizeY)
+void FViewport::FHitProxyMap::Init(FRHICommandListBase& RHICmdList, uint32 NewSizeX,uint32 NewSizeY)
 {
 	SizeX = NewSizeX;
 	SizeY = NewSizeY;
@@ -2282,7 +2282,7 @@ void FViewport::FHitProxyMap::Init(uint32 NewSizeX,uint32 NewSizeY)
 			.SetClearValue(FClearValueBinding::White)
 			.SetInitialState(ERHIAccess::SRVMask);
 
-		RenderTargetTextureRHI = RHICreateTexture(Desc);
+		RenderTargetTextureRHI = RHICmdList.CreateTexture(Desc);
 	}
 	{
 		const FRHITextureCreateDesc Desc =
@@ -2291,10 +2291,15 @@ void FViewport::FHitProxyMap::Init(uint32 NewSizeX,uint32 NewSizeY)
 			.SetFormat(PF_B8G8R8A8)
 			.SetFlags(ETextureCreateFlags::CPUReadback);
 
-		HitProxyCPUTexture = RHICreateTexture(Desc);
+		HitProxyCPUTexture = RHICmdList.CreateTexture(Desc);
 	}
 }
 
+void FViewport::FHitProxyMap::Init(uint32 NewSizeX, uint32 NewSizeY)
+{
+	Init(FRHICommandListImmediate::Get(), NewSizeX, NewSizeY);
+}
+
 void FViewport::FHitProxyMap::Release()
 {
 	HitProxyCPUTexture.SafeRelease();
diff --git a/Engine/Source/Runtime/Engine/Private/UnrealEngine.cpp b/Engine/Source/Runtime/Engine/Private/UnrealEngine.cpp
--- a/Engine/Source/Runtime/Engine/Private/UnrealEngine.cpp
+++ b/Engine/Source/Runtime/Engine/Private/UnrealEngine.cpp
@@ -4186,7 +4186,7 @@ public:
 
 	virtual void UpdateViewport(bool bUseSeparateRenderTarget, const class FViewport& Viewport, class SViewport* ViewportWidget = nullptr) override {};
 
-	virtual bool AllocateRenderTargetTextures(uint32 SizeX, uint32 SizeY, uint8 Format, uint32 NumLayers, ETextureCreateFlags Flags, ETextureCreateFlags TargetableTextureFlags, TArray<FTextureRHIRef>& OutTargetableTextures, TArray<FTextureRHIRef>& OutShaderResourceTextures, uint32 NumSamples = 1)
+	virtual bool AllocateRenderTargetTextures(FRHICommandListBase& RHICmdList, uint32 SizeX, uint32 SizeY, uint8 Format, uint32 NumLayers, ETextureCreateFlags Flags, ETextureCreateFlags TargetableTextureFlags, TArray<FTextureRHIRef>& OutTargetableTextures, TArray<FTextureRHIRef>& OutShaderResourceTextures, uint32 NumSamples = 1)
 	{ 
 		check(IsInRenderingThread());
 		checkf(bLayeredRTs, TEXT("Non-layered stereo can use the backbuffer directly"));
@@ -4229,7 +4229,7 @@ public:
 		SwapchainImageDesc.SetFlags(UnifiedCreateFlags);
 
 		// this will be owned by the caller
-		FTextureRHIRef NewImage = RHICreateTexture(SwapchainImageDesc);
+		FTextureRHIRef NewImage = RHICmdList.CreateTexture(SwapchainImageDesc);
 
 		OutTargetableTextures.Add(NewImage);
 		OutShaderResourceTextures = OutTargetableTextures;
diff --git a/Engine/Source/Runtime/Engine/Public/StereoRenderTargetManager.h b/Engine/Source/Runtime/Engine/Public/StereoRenderTargetManager.h
--- a/Engine/Source/Runtime/Engine/Public/StereoRenderTargetManager.h
+++ b/Engine/Source/Runtime/Engine/Public/StereoRenderTargetManager.h
@@ -5,10 +5,12 @@ StereoRenderTargetManager.h: Abstract interface returned from IStereoRendering t
 =============================================================================*/
 
 #pragma once
+
 #include "CoreMinimal.h"
 #include "Engine/GameViewportClient.h"
 #include "PixelFormat.h"
 
+class FRHICommandListBase;
 class FRHITexture;
 enum class ETextureCreateFlags : uint64;
 enum EShaderPlatform : uint16;
@@ -49,39 +51,25 @@ public:
 	 */
 	virtual bool NeedReAllocateViewportRenderTarget(const class FViewport& Viewport) = 0;
 
-	/**
-	* Returns true, if render target texture must be re-calculated.
-	*/
-	UE_DEPRECATED(5.2, "Return true in NeedReAllocateViewportRenderTarget instead")
-	virtual bool NeedReAllocateDepthTexture(const TRefCountPtr<struct IPooledRenderTarget>& DepthTarget) { return false; }
-
 	/**
 	* Returns true, if shading rate texture must be re-calculated.
 	*/
 	virtual bool NeedReAllocateShadingRateTexture(const TRefCountPtr<struct IPooledRenderTarget>& ShadingRateTarget) { return false; }
 
-	/**
-	 * Returns number of required buffered frames.
-	 */
-	UE_DEPRECATED(5.2, "Inferred from the array size returned in AllocateRenderTargetTextures")
-	virtual uint32 GetNumberOfBufferedFrames() const { return 1; }
-
-	/**
-	 * Allocates a render target texture.
-	 * The default implementation always return false to indicate that the default texture allocation should be used instead.
-	 *
-	 * @param Index			(in) index of the buffer, changing from 0 to GetNumberOfBufferedFrames()
-	 * @return				true, if texture was allocated; false, if the default texture allocation should be used.
-	 */
-	UE_DEPRECATED(5.2, "Implement AllocateRenderTargetTextures to allocate all textures at once")
-	virtual bool AllocateRenderTargetTexture(uint32 Index, uint32 SizeX, uint32 SizeY, uint8 Format, uint32 NumMips, ETextureCreateFlags Flags, ETextureCreateFlags TargetableTextureFlags, FTextureRHIRef& OutTargetableTexture, FTextureRHIRef& OutShaderResourceTexture, uint32 NumSamples = 1) { return false; }
-
 	/**
 	 * Allocates the render target textures, which includes the color textures and optionally other textures like depth.
 	 * The default implementation always returns false to indicate that the default texture allocation should be used instead.
 	 *
 	 * @return				true, if textures were allocated; false, if the default texture allocation should be used.
 	 */
+	virtual bool AllocateRenderTargetTextures(FRHICommandListBase& RHICmdList, uint32 SizeX, uint32 SizeY, uint8 Format, uint32 NumLayers, ETextureCreateFlags Flags, ETextureCreateFlags TargetableTextureFlags, TArray<FTextureRHIRef>& OutTargetableTextures, TArray<FTextureRHIRef>& OutShaderResourceTextures, uint32 NumSamples = 1)
+	{
+		PRAGMA_DISABLE_DEPRECATION_WARNINGS
+		return AllocateRenderTargetTextures(SizeX, SizeY, Format, NumLayers, Flags, TargetableTextureFlags, OutTargetableTextures, OutShaderResourceTextures, NumSamples);
+		PRAGMA_ENABLE_DEPRECATION_WARNINGS
+	}
+
+	UE_DEPRECATED(5.8, "RHI command list is now a required argument")
 	virtual bool AllocateRenderTargetTextures(uint32 SizeX, uint32 SizeY, uint8 Format, uint32 NumLayers, ETextureCreateFlags Flags, ETextureCreateFlags TargetableTextureFlags, TArray<FTextureRHIRef>& OutTargetableTextures, TArray<FTextureRHIRef>& OutShaderResourceTextures, uint32 NumSamples = 1) { return false; }
 
 	/**
diff --git a/Engine/Source/Runtime/Engine/Public/UnrealClient.h b/Engine/Source/Runtime/Engine/Public/UnrealClient.h
--- a/Engine/Source/Runtime/Engine/Public/UnrealClient.h
+++ b/Engine/Source/Runtime/Engine/Public/UnrealClient.h
@@ -766,6 +766,9 @@ protected:
 		ENGINE_API virtual ~FHitProxyMap();
 
 		/** Initializes the hit proxy map with the given dimensions. */
+		void Init(FRHICommandListBase& RHICmdList, uint32 NewSizeX, uint32 NewSizeY);
+
+		UE_DEPRECATED(5.8, "RHI command list is now a required argument")
 		void Init(uint32 NewSizeX,uint32 NewSizeY);
 
 		/** Releases the hit proxy resources. */
