XR Crash reporting improvement
-For OpenXR this will include the runtime and it's version.  For non-OpenXR it will tell us if we are on ARKit, ARCore, SimpleHMD, etc.
-Example: <XR>XR OpenXR PicoStreamingXR (Standalone) 1.1.46</XR>
#Synced-CL 47939539

diff --git a/Engine/Plugins/Experimental/SimpleHMD/Source/SimpleHMD/Private/SimpleHMD.cpp b/Engine/Plugins/Experimental/SimpleHMD/Source/SimpleHMD/Private/SimpleHMD.cpp
--- a/Engine/Plugins/Experimental/SimpleHMD/Source/SimpleHMD/Private/SimpleHMD.cpp
+++ b/Engine/Plugins/Experimental/SimpleHMD/Source/SimpleHMD/Private/SimpleHMD.cpp
@@ -6,6 +6,7 @@
 #include "ISimpleHMDPlugin.h"
 #include "PostProcess/PostProcessHMD.h"
 #include "GameFramework/WorldSettings.h"
+#include "GenericPlatform/GenericPlatformCrashContext.h"
 #include "SceneView.h"
 #include "RHIResourceUtils.h"
 
@@ -31,6 +32,7 @@ TSharedPtr< class IXRTrackingSystem, ESPMode::ThreadSafe > FSimpleHMDPlugin::Cre
 	auto SimpleHMD = FSceneViewExtensions::NewExtension<FSimpleHMD>();
 	if( SimpleHMD->IsInitialized() )
 	{
+		FGenericCrashContext::SetEngineData(TEXT("XR"), TEXT("XR SimpleHMD"));
 		return SimpleHMD;
 	}
 	return nullptr;
diff --git a/Engine/Plugins/Media/PixelStreaming/Source/PixelStreamingHMD/Private/PixelStreamingHMDModule.cpp b/Engine/Plugins/Media/PixelStreaming/Source/PixelStreamingHMD/Private/PixelStreamingHMDModule.cpp
--- a/Engine/Plugins/Media/PixelStreaming/Source/PixelStreamingHMD/Private/PixelStreamingHMDModule.cpp
+++ b/Engine/Plugins/Media/PixelStreaming/Source/PixelStreamingHMD/Private/PixelStreamingHMDModule.cpp
@@ -4,6 +4,7 @@
 #include "IXRTrackingSystem.h"
 #include "PixelStreamingHMD.h"
 #include "Settings.h"
+#include "GenericPlatform/GenericPlatformCrashContext.h"
 
 namespace UE::PixelStreamingHMD
 {
@@ -40,6 +41,7 @@ namespace UE::PixelStreamingHMD
 		TSharedRef<FPixelStreamingHMD> PixelStreamingHMD = FSceneViewExtensions::NewExtension<FPixelStreamingHMD>();
 		if (PixelStreamingHMD->IsInitialized())
 		{
+			FGenericCrashContext::SetEngineData(TEXT("XR"), TEXT("XR PixelStreaming"));
 			return PixelStreamingHMD;
 		}
 		return nullptr;
diff --git a/Engine/Plugins/Media/PixelStreaming2/Source/PixelStreaming2HMD/Private/PixelStreaming2HMDModule.cpp b/Engine/Plugins/Media/PixelStreaming2/Source/PixelStreaming2HMD/Private/PixelStreaming2HMDModule.cpp
--- a/Engine/Plugins/Media/PixelStreaming2/Source/PixelStreaming2HMD/Private/PixelStreaming2HMDModule.cpp
+++ b/Engine/Plugins/Media/PixelStreaming2/Source/PixelStreaming2HMD/Private/PixelStreaming2HMDModule.cpp
@@ -4,6 +4,7 @@
 
 #include "PixelStreaming2HMD.h"
 #include "PixelStreaming2PluginSettings.h"
+#include "GenericPlatform/GenericPlatformCrashContext.h"
 
 namespace UE::PixelStreaming2HMD
 {
@@ -39,6 +40,7 @@ namespace UE::PixelStreaming2HMD
 		TSharedRef<FPixelStreaming2HMD> PixelStreaming2HMD = FSceneViewExtensions::NewExtension<FPixelStreaming2HMD>();
 		if (PixelStreaming2HMD->IsInitialized())
 		{
+			FGenericCrashContext::SetEngineData(TEXT("XR"), TEXT("XR PixelStreaming2"));
 			return PixelStreaming2HMD;
 		}
 		return nullptr;
diff --git a/Engine/Plugins/Runtime/AR/AppleAR/AppleARKit/Source/AppleARKit/Private/AppleARKitModule.cpp b/Engine/Plugins/Runtime/AR/AppleAR/AppleARKit/Source/AppleARKit/Private/AppleARKitModule.cpp
--- a/Engine/Plugins/Runtime/AR/AppleAR/AppleARKit/Source/AppleARKit/Private/AppleARKitModule.cpp
+++ b/Engine/Plugins/Runtime/AR/AppleAR/AppleARKit/Source/AppleARKit/Private/AppleARKitModule.cpp
@@ -9,6 +9,7 @@
 #include "Misc/CoreDelegates.h"
 #include "Interfaces/IPluginManager.h"
 #include "ShaderCore.h"
+#include "GenericPlatform/GenericPlatformCrashContext.h"
 
 #if WITH_EDITOR
 	#include "ISettingsModule.h"
@@ -26,6 +27,7 @@ TSharedPtr<class IXRTrackingSystem, ESPMode::ThreadSafe> FAppleARKitModule::Crea
 	auto NewARKitSystem = AppleARKitSupport::CreateAppleARKitSystem();
 	NewARKitSystem->GetARCompositionComponent()->InitializeARSystem();
 	FAppleARKitARKitSystemPtr = NewARKitSystem;
+	FGenericCrashContext::SetEngineData(TEXT("XR"), TEXT("XR AppleARKit"));
     return NewARKitSystem;
 #else
 	return TSharedPtr<class IXRTrackingSystem, ESPMode::ThreadSafe>();
diff --git a/Engine/Plugins/Runtime/AR/Google/GoogleARCore/Source/GoogleARCoreBase/Private/GoogleARCoreBaseModule.cpp b/Engine/Plugins/Runtime/AR/Google/GoogleARCore/Source/GoogleARCoreBase/Private/GoogleARCoreBaseModule.cpp
--- a/Engine/Plugins/Runtime/AR/Google/GoogleARCore/Source/GoogleARCoreBase/Private/GoogleARCoreBaseModule.cpp
+++ b/Engine/Plugins/Runtime/AR/Google/GoogleARCore/Source/GoogleARCoreBase/Private/GoogleARCoreBaseModule.cpp
@@ -65,6 +65,7 @@ TSharedPtr< class IXRTrackingSystem, ESPMode::ThreadSafe > FGoogleARCoreBaseModu
 	auto ARCoreSystem = MakeShared<FGoogleARCoreXRTrackingSystem, ESPMode::ThreadSafe>();
 	ARCoreSystem->GetARCompositionComponent()->InitializeARSystem();
 	FGoogleARCoreDevice::GetInstance()->SetARSystem(ARCoreSystem->GetARCompositionComponent());
+	FGenericCrashContext::SetEngineData(TEXT("XR"), TEXT("XR GoogleARCore"));
 	return ARCoreSystem;
 #else
 	return TSharedPtr<class IXRTrackingSystem, ESPMode::ThreadSafe>();
diff --git a/Engine/Plugins/Runtime/OpenXR/Source/OpenXRHMD/Private/OpenXRHMDModule.cpp b/Engine/Plugins/Runtime/OpenXR/Source/OpenXRHMD/Private/OpenXRHMDModule.cpp
--- a/Engine/Plugins/Runtime/OpenXR/Source/OpenXRHMD/Private/OpenXRHMDModule.cpp
+++ b/Engine/Plugins/Runtime/OpenXR/Source/OpenXRHMD/Private/OpenXRHMDModule.cpp
@@ -11,6 +11,7 @@
 #include "OpenXRHMDSettings.h"
 #include "BuildSettings.h"
 #include "GeneralProjectSettings.h"
+#include "GenericPlatform/GenericPlatformCrashContext.h"
 #include "Epic_openxr.h"
 
 #if PLATFORM_ANDROID
@@ -865,6 +866,8 @@ bool FOpenXRHMDModule::InitInstance()
 	InstanceProps.runtimeName[XR_MAX_RUNTIME_NAME_SIZE - 1] = 0; // Ensure the name is null terminated.
 	UE_LOG(LogHMD, Log, TEXT("Initialized OpenXR on %S runtime version %d.%d.%d"), InstanceProps.runtimeName, XR_VERSION_MAJOR(InstanceProps.runtimeVersion), XR_VERSION_MINOR(InstanceProps.runtimeVersion), XR_VERSION_PATCH(InstanceProps.runtimeVersion));
 
+	FGenericCrashContext::SetEngineData(TEXT("XR"), FString::Printf(TEXT("XR OpenXRHMD %hs %d.%d.%d"), InstanceProps.runtimeName, XR_VERSION_MAJOR(InstanceProps.runtimeVersion), XR_VERSION_MINOR(InstanceProps.runtimeVersion), XR_VERSION_PATCH(InstanceProps.runtimeVersion)));
+
 	if (CVarCheckOpenXRInstanceConformance.GetValueOnAnyThread() &&
 		(FCStringAnsi::Strstr(InstanceProps.runtimeName, "SteamVR/OpenXR") != nullptr))
 	{
