[CrashReport][Assert] OpenXRInput!FOpenXRInputPlugin::FOpenXRAction::FOpenXRAction() [OpenXRInput.cpp]
-Added an ensure and early out to report wide character action names which openxr does not support, and then ignore them to avoid the crash.
#Synced-CL 47939539

diff --git a/Engine/Plugins/Runtime/OpenXR/Source/OpenXRInput/Private/OpenXRInput.cpp b/Engine/Plugins/Runtime/OpenXR/Source/OpenXRInput/Private/OpenXRInput.cpp
--- a/Engine/Plugins/Runtime/OpenXR/Source/OpenXRInput/Private/OpenXRInput.cpp
+++ b/Engine/Plugins/Runtime/OpenXR/Source/OpenXRInput/Private/OpenXRInput.cpp
@@ -690,10 +690,15 @@ void FOpenXRInputPlugin::FOpenXRInput::BuildEnhancedActions(TMap<FString, FInter
 			if (ActionIndex == INDEX_NONE)
 			{
 				// No action found, create a new one
-				FString NewActionDescription = Mapping.Action->ActionDescription.ToString();
 
-				// Handle any description string duplicates
+				if (!ensureMsgf(!NewActionName.GetDisplayNameEntry()->IsWide(), TEXT("Input Action %s has wide characters in its name, This name will be passed to OpenXR which does not support wide characters in path names (see spec).  Ignoring this action, it will not function."), *NewActionName.ToString()))
+				{
+					continue;
+				}
+
+				FString NewActionDescription = Mapping.Action->ActionDescription.ToString();
 				{
+					// Handle any description string duplicates
 					FName* ExistingActionName = ActionDescriptionMap.Find(NewActionDescription);
 					if (ExistingActionName != nullptr)
 					{
