OpenXR 1.1 config variable lookup not working
-We were looking in the game ini and needed to look in the engine ini.
-This probably is not affecting very much as the main runtime seem to only support 1.0.

#Synced-CL 49564021

diff --git a/Engine/Plugins/Runtime/OpenXR/Source/OpenXRHMD/Private/OpenXRHMDModule.cpp b/Engine/Plugins/Runtime/OpenXR/Source/OpenXRHMD/Private/OpenXRHMDModule.cpp
--- a/Engine/Plugins/Runtime/OpenXR/Source/OpenXRHMD/Private/OpenXRHMDModule.cpp
+++ b/Engine/Plugins/Runtime/OpenXR/Source/OpenXRHMD/Private/OpenXRHMDModule.cpp
@@ -926,11 +926,11 @@ bool FOpenXRHMDModule::TryCreateInstance(XrInstanceCreateInfo& Info)
 		// Read the ini.  This code can run so early that we cannot use a default object to load this more easily.
 		// Note the settings are true by default, so if we don't find a setting we use true.
 		bool bFoundValue = false;
-		if (GConfig->GetBool(TEXT("/Script/OpenXRHMD.OpenXRHMDSettings"), TEXT("bIsOpenXR1_1Enabled"), bFoundValue, GGameIni))
+		if (GConfig->GetBool(TEXT("/Script/OpenXRHMD.OpenXRHMDSettings"), TEXT("bIsOpenXR1_1Enabled"), bFoundValue, GEngineIni))
 		{
 			bIsOpenXR1_1Enabled = bFoundValue;
 		}
-		if (GConfig->GetBool(TEXT("/Script/OpenXRHMD.OpenXRHMDSettings"), TEXT("bIsOpenXR1_0Enabled"), bFoundValue, GGameIni))
+		if (GConfig->GetBool(TEXT("/Script/OpenXRHMD.OpenXRHMDSettings"), TEXT("bIsOpenXR1_0Enabled"), bFoundValue, GEngineIni))
 		{
 			bIsOpenXR1_0Enabled = bFoundValue;
 		}
