PR #13488: [OpenXR]Add equirect2 layer radius parameter.

#Synced-CL 49639717

diff --git a/Engine/Plugins/Runtime/OpenXR/Source/OpenXRHMD/Private/OpenXRHMD_Layer.cpp b/Engine/Plugins/Runtime/OpenXR/Source/OpenXRHMD/Private/OpenXRHMD_Layer.cpp
--- a/Engine/Plugins/Runtime/OpenXR/Source/OpenXRHMD/Private/OpenXRHMD_Layer.cpp
+++ b/Engine/Plugins/Runtime/OpenXR/Source/OpenXRHMD/Private/OpenXRHMD_Layer.cpp
@@ -179,10 +179,14 @@ void FOpenXRLayer::CreateOpenXREquirect2Layer(bool bIsStereo, bool bNoAlpha, FTr
 	const FEquirectLayer& EquirectProps = Desc.GetShape<FEquirectLayer>();
 
 	XrCompositionLayerEquirect2KHR Equirect2 = { XR_TYPE_COMPOSITION_LAYER_EQUIRECT2_KHR, nullptr };
-	Equirect2.layerFlags = bNoAlpha ? 0 : XR_COMPOSITION_LAYER_UNPREMULTIPLIED_ALPHA_BIT |
-		XR_COMPOSITION_LAYER_BLEND_TEXTURE_SOURCE_ALPHA_BIT;
+	Equirect2.layerFlags = 
+		bNoAlpha ? 
+		0 :
+		XR_COMPOSITION_LAYER_UNPREMULTIPLIED_ALPHA_BIT | XR_COMPOSITION_LAYER_BLEND_TEXTURE_SOURCE_ALPHA_BIT;
+	
 	Equirect2.space = Space;
 	Equirect2.subImage.imageArrayIndex = 0;
+	Equirect2.radius = FMath::Abs(EquirectProps.Radius / WorldToMeters);
 	FXrCompositionLayerUnion LayerUnion;
 
 	if (RightEye.Swapchain.IsValid())
