OpenXR initialized by default on Android platforms, if available
-Normally we would like enabling the OpenXR plugin to be the decision, however long ago the openxr loader would hang if openxr was not available on some android devices and, notably, qagame.  So we made the default for this true.
-The OpenXR loader has been through many updates and our lowest supported android version is a lot higher now and there are more non-oculus android xr devices.  We did not find a device that hangs in testing.

#Synced-CL 50171876

diff --git a/Engine/Plugins/Runtime/OpenXR/Source/OpenXRHMD/Private/OpenXRHMDModule.cpp b/Engine/Plugins/Runtime/OpenXR/Source/OpenXRHMD/Private/OpenXRHMDModule.cpp
--- a/Engine/Plugins/Runtime/OpenXR/Source/OpenXRHMD/Private/OpenXRHMDModule.cpp
+++ b/Engine/Plugins/Runtime/OpenXR/Source/OpenXRHMD/Private/OpenXRHMDModule.cpp
@@ -30,8 +30,8 @@ static TAutoConsoleVariable<int32> CVarEnableOpenXRValidationLayer(
 
 static TAutoConsoleVariable<bool> CVarDisableOpenXROnAndroidWithoutOculus(
 	TEXT("xr.DisableOpenXROnAndroidWithoutOculus"),
-	true,
-	TEXT("If true OpenXR will not initialize on Android unless the project is packaged for Oculus (ProjectSetting->Platforms->Android->Advanced APK Packaging->PackageForOculusMobileDevices list not empty).  Currently defaulted to true because the OpenXR loader we are using hangs during intialization on some devices instead of failing, as it should."),
+	false,
+	TEXT("If true OpenXR will not initialize on Android unless the project is packaged for Oculus (ProjectSetting->Platforms->Android->Advanced APK Packaging->PackageForOculusMobileDevices list not empty)."),
 	ECVF_RenderThreadSafe);
 
 static TAutoConsoleVariable<bool> CVarCheckOpenXRInstanceConformance(
