OpenXR: Fix for parsing of Input dpad paths
-Some special handling for dpad inputs because it's natural keynames do not quite fix what our openxr input code expects.
#Synced-CL 50220247

diff --git a/Engine/Plugins/Runtime/OpenXR/Source/OpenXRInput/Private/OpenXRInput.cpp b/Engine/Plugins/Runtime/OpenXR/Source/OpenXRInput/Private/OpenXRInput.cpp
--- a/Engine/Plugins/Runtime/OpenXR/Source/OpenXRInput/Private/OpenXRInput.cpp
+++ b/Engine/Plugins/Runtime/OpenXR/Source/OpenXRInput/Private/OpenXRInput.cpp
@@ -822,6 +822,21 @@ bool FOpenXRInputPlugin::FOpenXRInput::SuggestBindingForKey(TMap<FString, FInter
 	{
 		Identifier = "squeeze";
 	}
+	else if(Identifier.StartsWith(TEXT("dpad")))
+	{
+		if(bDirectionalBindingSupported)
+		{
+			// dpad inputs are specified as <ProfileName>_<Hand>_DPad<Direction>_<action>.  Normally, we'd have a _ between
+			// "dpad" and the direction of the button, but in order to match the hard NUM_XR_KEY_TOKENS
+			// formatting requirement, we detect if the third Token starts with "dpad," and extract the direction in order
+			// to form the longer path required.
+			Identifier = "dpad_" + Identifier.Mid(4);
+		}
+		else
+		{
+			return false;
+		}
+	}
 	Path += "/input/" + Identifier;
 
 	// Add the data we want to query, we'll skip this for trigger/squeeze "click" actions to allow
