IOpenXRExtensionPlugin GetBlockedOptionalExtensions
-Added GetBlockedOptionalExtensions which fills an array with extensions to block.
-Removing those extensions from the optional extension requests by UE and by OpenXRExtensionPlugins.
-Logging so we know what is happening.

#Synced-CL 50693347

diff --git a/Engine/Plugins/Runtime/OpenXR/Source/OpenXRHMD/Private/OpenXRHMDModule.cpp b/Engine/Plugins/Runtime/OpenXR/Source/OpenXRHMD/Private/OpenXRHMDModule.cpp
--- a/Engine/Plugins/Runtime/OpenXR/Source/OpenXRHMD/Private/OpenXRHMDModule.cpp
+++ b/Engine/Plugins/Runtime/OpenXR/Source/OpenXRHMD/Private/OpenXRHMDModule.cpp
@@ -31,7 +31,7 @@ static TAutoConsoleVariable<int32> CVarEnableOpenXRValidationLayer(
 static TAutoConsoleVariable<bool> CVarDisableOpenXROnAndroidWithoutOculus(
 	TEXT("xr.DisableOpenXROnAndroidWithoutOculus"),
 	false,
-	TEXT("If true OpenXR will not initialize on Android unless the project is packaged for Oculus (ProjectSetting->Platforms->Android->Advanced APK Packaging->PackageForOculusMobileDevices list not empty)."),
+	TEXT("If true OpenXR will not initialize on Android unless the project is packaged for Oculus (ProjectSetting->Platforms->Android->Advanced APK Packaging->PackageForMetaQuest)."),
 	ECVF_RenderThreadSafe);
 
 static TAutoConsoleVariable<bool> CVarCheckOpenXRInstanceConformance(
@@ -330,38 +330,42 @@ bool FOpenXRHMDModule::EnumerateLayers()
 	return false;
 }
 
-struct AnsiKeyFunc : BaseKeyFuncs<const ANSICHAR*, const ANSICHAR*, false>
+namespace OpenXRHMDModule
 {
-	typedef typename TTypeTraits<const ANSICHAR*>::ConstPointerType KeyInitType;
-	typedef typename TCallTraits<const ANSICHAR*>::ParamType ElementInitType;
-
-	/**
-	 * @return The key used to index the given element.
-	 */
-	static FORCEINLINE KeyInitType GetSetKey(ElementInitType Element)
+	struct AnsiKeyFunc : BaseKeyFuncs<const ANSICHAR*, const ANSICHAR*, false>
 	{
-		return Element;
-	}
+		typedef typename TTypeTraits<const ANSICHAR*>::ConstPointerType KeyInitType;
+		typedef typename TCallTraits<const ANSICHAR*>::ParamType ElementInitType;
 
-	/**
-	 * @return True if the keys match.
-	 */
-	static FORCEINLINE bool Matches(KeyInitType A, KeyInitType B)
-	{
-		return FCStringAnsi::Strcmp(A, B) == 0;
-	}
+		/**
+		 * @return The key used to index the given element.
+		 */
+		static FORCEINLINE KeyInitType GetSetKey(ElementInitType Element)
+		{
+			return Element;
+		}
+
+		/**
+		 * @return True if the keys match.
+		 */
+		static FORCEINLINE bool Matches(KeyInitType A, KeyInitType B)
+		{
+			return FCStringAnsi::Strcmp(A, B) == 0;
+		}
+
+		/** Calculates a hash index for a key. */
+		static FORCEINLINE uint32 GetKeyHash(KeyInitType Key)
+		{
+			return UE::HashStringFNV1a32(MakeStringView(Key));
+		}
+	};
+}
 
-	/** Calculates a hash index for a key. */
-	static FORCEINLINE uint32 GetKeyHash(KeyInitType Key)
-	{
-		return GetTypeHash(Key);
-	}
-};
 
 bool FOpenXRHMDModule::InitRenderBridge()
 {
 	// Get all extension plugins
-	TSet<const ANSICHAR*, AnsiKeyFunc> ExtensionSet;
+	TSet<const ANSICHAR*, OpenXRHMDModule::AnsiKeyFunc> ExtensionSet;
 	TArray<IOpenXRExtensionPlugin*> ExtModules = IModularFeatures::Get().GetModularFeatureImplementations<IOpenXRExtensionPlugin>(IOpenXRExtensionPlugin::GetModularFeatureName());
 
 	// Query all extension plugins to see if we need to use a custom render bridge
@@ -637,6 +641,58 @@ bool FOpenXRHMDModule::GetOptionalExtensions(TArray<const ANSICHAR*>& OutExtensi
 	return true;
 }
 
+namespace OpenXRHMDModule
+{
+	typedef TSet<const ANSICHAR*, OpenXRHMDModule::AnsiKeyFunc> ExtensionSet;
+	bool BlockOptionalExtensions(TArray<const ANSICHAR*>& InOutExtensions, const ExtensionSet& BlockedExtensions)
+	{
+		if (BlockedExtensions.Num() == 0)
+		{
+			return false;
+		}
+
+		const int32 InExtensionsCount = InOutExtensions.Num();
+		int32 Cursor = 0;
+		int32 End = InExtensionsCount - 1;
+		while (Cursor <= End)
+		{
+			const ANSICHAR* Extension = InOutExtensions[Cursor];
+			if (!BlockedExtensions.Contains(Extension))
+			{
+				++Cursor;
+			}
+			else
+			{
+				UE_LOG(LogHMD, Log, TEXT("OpenXRHMDModule::BlockOptionalExtensions blocking % hs"), InOutExtensions[Cursor]);
+				InOutExtensions[Cursor] = InOutExtensions[End];
+				--End;
+			}
+		}
+		check(End == Cursor - 1);
+		InOutExtensions.SetNum(Cursor);
+
+		return Cursor != InExtensionsCount;
+	}
+
+	void LogExtensions(const TCHAR* Preface, const TArray<const ANSICHAR*>& Extensions)
+	{
+		UE_LOG(LogHMD, Log, TEXT("%s"), Preface);
+		for (const ANSICHAR* Extension : Extensions)
+		{
+			UE_LOG(LogHMD, Log, TEXT("      %hs"), Extension);
+		}
+	}
+	
+	void LogExtensions(bool DoLog, const TCHAR* Preface, const TArray<const ANSICHAR*>& Extensions)
+	{
+		if (DoLog)
+		{
+			LogExtensions(Preface, Extensions);
+		}
+	}
+
+}
+
 bool FOpenXRHMDModule::InitInstance()
 {
 	if (Instance)
@@ -664,7 +720,7 @@ bool FOpenXRHMDModule::InitInstance()
 #endif
 
 	// Get all extension plugins
-	TSet<const ANSICHAR*, AnsiKeyFunc> ExtensionSet;
+	OpenXRHMDModule::ExtensionSet ExtensionSet;
 	TArray<IOpenXRExtensionPlugin*> ExtModules = IModularFeatures::Get().GetModularFeatureImplementations<IOpenXRExtensionPlugin>(IOpenXRExtensionPlugin::GetModularFeatureName());
 
 	// Query all extension plugins to see if we need to use a custom loader
@@ -714,8 +770,46 @@ bool FOpenXRHMDModule::InitInstance()
 		return false;
 	}
 
+	// Gather set of optional extensions that are blocked by extension plugins
+	OpenXRHMDModule::ExtensionSet BlockedOptionalExtensionSet;
+	{
+		TMap<IOpenXRExtensionPlugin*, TArray<const ANSICHAR*>> BlockedOptionalExtensionMap;
+		for (IOpenXRExtensionPlugin* Plugin : ExtModules)
+		{
+			TArray<const ANSICHAR*> BlockedOptionalExtensions;
+			bool RetVal = Plugin->GetBlockedOptionalExtensions(BlockedOptionalExtensions);
+			if (RetVal && BlockedOptionalExtensions.Num() > 0)
+			{
+				BlockedOptionalExtensionMap.Add(Plugin, BlockedOptionalExtensions);
+				for (const ANSICHAR* Blocked : BlockedOptionalExtensions)
+				{
+					BlockedOptionalExtensionSet.Add(Blocked);
+				}
+			}
+		}
+		// Logging to show which extension plugins are blocking what optional extensions.
+		UE_LOG(LogHMD, Log, TEXT("BlockedOptionalExtensions:"));
+		if (BlockedOptionalExtensionMap.IsEmpty())
+		{
+			UE_LOG(LogHMD, Log, TEXT("  None"));
+		}
+		else
+		{
+			for (TPair< IOpenXRExtensionPlugin*, TArray<const ANSICHAR*> > Pair : BlockedOptionalExtensionMap)
+			{
+				UE_LOG(LogHMD, Log, TEXT("  %s is blocking:"), *(Pair.Key->GetDisplayName()));
+				for (const ANSICHAR* Extension : Pair.Value)
+				{
+					UE_LOG(LogHMD, Log, TEXT("    %hs"), Extension);
+				}
+			}
+		}
+	}
+
 	// Enable any required and optional extensions that are not plugin specific (usually platform support extensions)
 	{
+		UE_LOG(LogHMD, Log, TEXT("Building UE core extension lists:"));
+
 		TArray<const ANSICHAR*> RequiredExtensions, OptionalExtensions, Extensions;
 		// Query required extensions
 		RequiredExtensions.Empty();
@@ -724,6 +818,7 @@ bool FOpenXRHMDModule::InitInstance()
 			UE_LOG(LogHMD, Error, TEXT("Could not get required OpenXR extensions."));
 			return false;
 		}
+		OpenXRHMDModule::LogExtensions(TEXT("  RequiredExtensions"), RequiredExtensions);
 
 		// Query optional extensions
 		OptionalExtensions.Empty();
@@ -732,6 +827,10 @@ bool FOpenXRHMDModule::InitInstance()
 			UE_LOG(LogHMD, Error, TEXT("Could not get optional OpenXR extensions."));
 			return false;
 		}
+		OpenXRHMDModule::LogExtensions(TEXT("  OptionalExtensions"), OptionalExtensions);
+
+		OpenXRHMDModule::BlockOptionalExtensions(OptionalExtensions, BlockedOptionalExtensionSet);
+		OpenXRHMDModule::LogExtensions(TEXT("  OptionalExtensions After Blocking:"), OptionalExtensions);
 
 		if (!EnableExtensions(RequiredExtensions, OptionalExtensions, Extensions))
 		{
@@ -746,9 +845,11 @@ bool FOpenXRHMDModule::InitInstance()
 		ExtensionSet.Add(XR_EPIC_VIEW_CONFIGURATION_FOV_EXTENSION_NAME);
 	}
 
+	UE_LOG(LogHMD, Log, TEXT("Extensions for IOpenXRExtensionPlugins:"));
 	ExtensionPlugins.Reset();
 	for (IOpenXRExtensionPlugin* Plugin : ExtModules)
 	{
+		UE_LOG(LogHMD, Log, TEXT("  %s:"), *Plugin->GetDisplayName());
 		TArray<const ANSICHAR*> RequiredExtensions, OptionalExtensions, Extensions;
 
 		// Query required extensions
@@ -758,6 +859,7 @@ bool FOpenXRHMDModule::InitInstance()
 			// Ignore the plugin if the query fails
 			continue;
 		}
+		OpenXRHMDModule::LogExtensions(!RequiredExtensions.IsEmpty(), TEXT("    RequiredExtensions"), RequiredExtensions);
 
 		// Query optional extensions
 		OptionalExtensions.Empty();
@@ -766,6 +868,11 @@ bool FOpenXRHMDModule::InitInstance()
 			// Ignore the plugin if the query fails
 			continue;
 		}
+		const bool bLogOptionalExtensions = !OptionalExtensions.IsEmpty();
+		OpenXRHMDModule::LogExtensions(bLogOptionalExtensions, TEXT("    OptionalExtensions"), OptionalExtensions);
+
+		OpenXRHMDModule::BlockOptionalExtensions(OptionalExtensions, BlockedOptionalExtensionSet);
+		OpenXRHMDModule::LogExtensions(bLogOptionalExtensions, TEXT("    OptionalExtensions After Blocking:"), OptionalExtensions);
 
 		if (!EnableExtensions(RequiredExtensions, OptionalExtensions, Extensions))
 		{
@@ -790,6 +897,8 @@ bool FOpenXRHMDModule::InitInstance()
 	{
 		EnabledExtensions.Add(Ext);
 	}
+	
+	OpenXRHMDModule::LogExtensions(TEXT("    EnabledExtensions:"), EnabledExtensions);
 
 	// Enable layers, if specified by CVar.
 	// Note: For the validation layer to work on Windows (as of latest OpenXR runtime, August 2019), the following are required:
diff --git a/Engine/Plugins/Runtime/OpenXR/Source/OpenXRHMD/Public/IOpenXRExtensionPlugin.h b/Engine/Plugins/Runtime/OpenXR/Source/OpenXRHMD/Public/IOpenXRExtensionPlugin.h
--- a/Engine/Plugins/Runtime/OpenXR/Source/OpenXRHMD/Public/IOpenXRExtensionPlugin.h
+++ b/Engine/Plugins/Runtime/OpenXR/Source/OpenXRHMD/Public/IOpenXRExtensionPlugin.h
@@ -213,6 +213,23 @@ public:
 		return true;
 	}
 
+	/**
+	* Fill the array with optional extensions blocked by this plugin.
+	* This can be used to disable optional extension usage by the UE OpenXR core and by other ExtensionPlugins.
+	* Required extensions either from UE OpenXR core or from ExtensionPlugins are NOT blocked.  In the latter case disabling plugins may have the
+	* same effect.  Otherwise code changes will be needed.
+	* In general blocking optional extensions used by UE's OpenXR core should be ok because it is designed to run correctly without those extensions
+	* and it seems likely the blocking plugin is replacing that functionality somehow.
+	* Blocking optional extensions added by plugins is potentially more complicated and ExtensionPlugins that do this may make themselves harder to use
+	* along with other ExtensionPlugins.  They are optional, so it should work, but feature use may become complex.
+	* Plugin settings to explicitly disable extension use or splitting plugins up so that features can be individually enabled/disable should be considered.
+	* If false is returned the returned extension list will be ignored.
+	*/
+	virtual bool GetBlockedOptionalExtensions(TArray<const ANSICHAR*>& OutExtensions)
+	{
+		return false;
+	}
+
 	/**
 	* Set the output parameters to add multiple interaction profiles to OpenXR Input
 	*/
