Ensure we do not allocate OpenXR motion vector and motion vector depth swapchains unless r.Velocity.DirectlyRenderOpenXRMotionVectors=1.

#Synced-CL 50823074

diff --git a/Engine/Plugins/Runtime/OpenXR/Source/OpenXRHMD/Private/OpenXRHMD.cpp b/Engine/Plugins/Runtime/OpenXR/Source/OpenXRHMD/Private/OpenXRHMD.cpp
--- a/Engine/Plugins/Runtime/OpenXR/Source/OpenXRHMD/Private/OpenXRHMD.cpp
+++ b/Engine/Plugins/Runtime/OpenXR/Source/OpenXRHMD/Private/OpenXRHMD.cpp
@@ -164,7 +164,7 @@ static TAutoConsoleVariable<bool> CVarOpenXRAlphaInvertPass(
 static TAutoConsoleVariable<bool> CVarOpenXRFrameSynthesis(
 	TEXT("xr.OpenXRFrameSynthesis"),
 	0,
-	TEXT("If true and supported via XR_EXT_frame_synthesis or XR_FB_space_warp, enable and submit motion vector and motion vector depth swapchains for frame synthesis.\n")
+	TEXT("If true and supported via XR_EXT_frame_synthesis or XR_FB_space_warp, write to and submit motion vector and motion vector depth swapchains for frame synthesis.\n")
 	TEXT("Currently only supported when using the Vulkan mobile renderer, using mobile multi-view, and r.Velocity.DirectlyRenderOpenXRMotionVectors=True.\n")
 	TEXT("Because normal velocity rendering is disabled when r.Velocity.DirectlyRenderOpenXRMotionVectors=True, temporal anti-aliasing and motion blur will be automatically disabled."),
 	ECVF_RenderThreadSafe);
@@ -2894,8 +2894,9 @@ bool FOpenXRHMD::AllocateRenderTargetTextures(FRHICommandListBase& RHICmdList, u
 		}
 	}
 
+	// Only create motion vector + motion vector depth swapchains if XR_EXT_frame_synthesis or XR_FB_space_warp are supported and r.Velocity.DirectlyRenderOpenXRMotionVectors=1
 	FIntPoint MotionVectorSwapchainSize = FIntPoint();
-	if (GetRecommendedMotionVectorTextureSize(MotionVectorSwapchainSize))
+	if (GetRecommendedMotionVectorTextureSize(MotionVectorSwapchainSize) && PlatformSupportsOpenXRMotionVectors(GetConfiguredShaderPlatform()))
 	{
 		const ETextureCreateFlags MotionVectorCreateFlags = TexCreate_RenderTargetable | TexCreate_ResolveTargetable | TexCreate_ShaderResource | TexCreate_InputAttachmentRead | TexCreate_Dynamic;
 		uint8 UnusedActualFormat = 0;
