OpenXRHMDSettings BlockedExtensions
-Added a config variable backed project setting to block the use of a list of OpenXRExtensions even if the runtime supports them and even if they are required for us to run OpenXR at all.

#Synced-CL 50693347

diff --git a/Engine/Plugins/Runtime/OpenXR/Source/OpenXRHMD/Private/OpenXRHMDModule.cpp b/Engine/Plugins/Runtime/OpenXR/Source/OpenXRHMD/Private/OpenXRHMDModule.cpp
--- a/Engine/Plugins/Runtime/OpenXR/Source/OpenXRHMD/Private/OpenXRHMDModule.cpp
+++ b/Engine/Plugins/Runtime/OpenXR/Source/OpenXRHMD/Private/OpenXRHMDModule.cpp
@@ -497,49 +497,6 @@ PFN_xrGetInstanceProcAddr FOpenXRHMDModule::GetDefaultLoader()
 	return OutGetProcAddr;
 }
 
-bool FOpenXRHMDModule::EnableExtensions(const TArray<const ANSICHAR*>& RequiredExtensions, const TArray<const ANSICHAR*>& OptionalExtensions, TArray<const ANSICHAR*>& OutExtensions)
-{
-	// Query required extensions and check if they're all available
-	bool ExtensionMissing = false;
-	for (const ANSICHAR* Ext : RequiredExtensions)
-	{
-		if (AvailableExtensions.Contains(Ext))
-		{
-			UE_LOG(LogHMD, Verbose, TEXT("Required extension %S enabled"), Ext);
-		}
-		else
-		{
-			UE_LOG(LogHMD, Warning, TEXT("Required extension %S is not available"), Ext);
-			ExtensionMissing = true;
-		}
-	}
-
-	// If any required extensions are missing then we ignore the plugin
-	if (ExtensionMissing)
-	{
-		return false;
-	}
-
-	// All required extensions are supported we can safely add them to our set and give the plugin callbacks
-	OutExtensions.Append(RequiredExtensions);
-
-	// Add all supported optional extensions to the set
-	for (const ANSICHAR* Ext : OptionalExtensions)
-	{
-		if (AvailableExtensions.Contains(Ext))
-		{
-			UE_LOG(LogHMD, Verbose, TEXT("Optional extension %S enabled"), Ext);
-			OutExtensions.Add(Ext);
-		}
-		else
-		{
-			UE_LOG(LogHMD, Log, TEXT("Optional extension %S is not available"), Ext);
-		}
-	}
-
-	return true;
-}
-
 bool FOpenXRHMDModule::GetRequiredExtensions(TArray<const ANSICHAR*>& OutExtensions)
 {
 #if PLATFORM_ANDROID
@@ -644,53 +601,91 @@ bool FOpenXRHMDModule::GetOptionalExtensions(TArray<const ANSICHAR*>& OutExtensi
 namespace OpenXRHMDModule
 {
 	typedef TSet<const ANSICHAR*, OpenXRHMDModule::AnsiKeyFunc> ExtensionSet;
-	bool BlockOptionalExtensions(TArray<const ANSICHAR*>& InOutExtensions, const ExtensionSet& BlockedExtensions)
+
+	bool EnableExtensions(const TSet<FString>& AvailableExtensions, const TArray<const ANSICHAR*>& RequiredExtensions, const TArray<const ANSICHAR*>& OptionalExtensions, TArray<const ANSICHAR*>& OutExtensions, const TSet<FString>& FullyBlockedExtensionsSet, const ExtensionSet& BlockedOptionalExtensionSet)
 	{
-		if (BlockedExtensions.Num() == 0)
+		// Query required extensions and check if they're all available
+		bool ExtensionMissing = false;
+		for (const ANSICHAR* Ext : RequiredExtensions)
+		{
+			if (AvailableExtensions.Contains(Ext))
+			{
+				if (FullyBlockedExtensionsSet.Contains(Ext))
+				{
+					UE_LOG(LogHMD, Warning, TEXT("Required extension %S is blocked by FullyBlockedExtensionsSet"), Ext);
+					ExtensionMissing = true;
+				}
+				else
+				{
+					UE_LOG(LogHMD, Verbose, TEXT("Required extension %S enabled"), Ext);
+				}
+			}
+			else
+			{
+				UE_LOG(LogHMD, Warning, TEXT("Required extension %S is not available"), Ext);
+				ExtensionMissing = true;
+			}
+		}
+
+		// If any required extensions are missing then we ignore the plugin
+		if (ExtensionMissing)
 		{
 			return false;
 		}
 
-		const int32 InExtensionsCount = InOutExtensions.Num();
-		int32 Cursor = 0;
-		int32 End = InExtensionsCount - 1;
-		while (Cursor <= End)
+		// All required extensions are supported we can safely add them to our set and give the plugin callbacks
+		OutExtensions.Append(RequiredExtensions);
+
+		// Add all supported optional extensions to the set, unless blocked.
+		for (const ANSICHAR* Ext : OptionalExtensions)
 		{
-			const ANSICHAR* Extension = InOutExtensions[Cursor];
-			if (!BlockedExtensions.Contains(Extension))
+			if (AvailableExtensions.Contains(Ext))
 			{
-				++Cursor;
+				if (FullyBlockedExtensionsSet.Contains(Ext))
+				{
+					UE_LOG(LogHMD, Log, TEXT("Optional extension %S is blocked by FullyBlockedExtensionsSet"), Ext);
+				}
+				else if (BlockedOptionalExtensionSet.Contains(Ext))
+				{
+					UE_LOG(LogHMD, Log, TEXT("Optional extension %S is blocked by BlockedOptionalExtensionSet"), Ext);
+				}
+				else
+				{
+					UE_LOG(LogHMD, Verbose, TEXT("Optional extension %S enabled"), Ext);
+					OutExtensions.Add(Ext);
+				}
 			}
 			else
 			{
-				UE_LOG(LogHMD, Log, TEXT("OpenXRHMDModule::BlockOptionalExtensions blocking % hs"), InOutExtensions[Cursor]);
-				InOutExtensions[Cursor] = InOutExtensions[End];
-				--End;
+				UE_LOG(LogHMD, Log, TEXT("Optional extension %S is not available"), Ext);
 			}
 		}
-		check(End == Cursor - 1);
-		InOutExtensions.SetNum(Cursor);
 
-		return Cursor != InExtensionsCount;
+		return true;
 	}
-
-	void LogExtensions(const TCHAR* Preface, const TArray<const ANSICHAR*>& Extensions)
+	
+	void LogExtensions(bool DoLog, const TCHAR* Preface, const TArray<const ANSICHAR*>& Extensions)
 	{
-		UE_LOG(LogHMD, Log, TEXT("%s"), Preface);
-		for (const ANSICHAR* Extension : Extensions)
+		if (DoLog)
 		{
-			UE_LOG(LogHMD, Log, TEXT("      %hs"), Extension);
+			UE_LOG(LogHMD, Log, TEXT("%s"), Preface);
+			for (const ANSICHAR* Extension : Extensions)
+			{
+				UE_LOG(LogHMD, Log, TEXT("      %hs"), Extension);
+			}
 		}
 	}
-	
-	void LogExtensions(bool DoLog, const TCHAR* Preface, const TArray<const ANSICHAR*>& Extensions)
+	void LogExtensionsVerbose(bool DoLog, const TCHAR* Preface, const TArray<const ANSICHAR*>& Extensions)
 	{
 		if (DoLog)
 		{
-			LogExtensions(Preface, Extensions);
+			UE_LOG(LogHMD, Verbose, TEXT("%s"), Preface);
+			for (const ANSICHAR* Extension : Extensions)
+			{
+				UE_LOG(LogHMD, Verbose, TEXT("      %hs"), Extension);
+			}
 		}
 	}
-
 }
 
 bool FOpenXRHMDModule::InitInstance()
@@ -770,6 +765,24 @@ bool FOpenXRHMDModule::InitInstance()
 		return false;
 	}
 
+	// Gather set of extensions that are blocked by OpenXRHMDSettings BlockedExtensions setting/ini.
+	TSet<FString> FullyBlockedExtensionsSet;
+	{
+		TArray<FString> BlockedExtensions;
+		if (GConfig->GetArray(TEXT("/Script/OpenXRHMD.OpenXRHMDSettings"), TEXT("BlockedExtensions"), BlockedExtensions, GEngineIni))
+		{
+			FullyBlockedExtensionsSet.Append(BlockedExtensions);
+		}
+	}
+	if (!FullyBlockedExtensionsSet.IsEmpty())
+	{
+		UE_LOG(LogHMD, Log, TEXT("FullyBlockedExtensionsSet:"));
+		for (const FString& BlockedExtension : FullyBlockedExtensionsSet)
+		{
+			UE_LOG(LogHMD, Log, TEXT("  %s"), *BlockedExtension);
+		}
+	}
+
 	// Gather set of optional extensions that are blocked by extension plugins
 	OpenXRHMDModule::ExtensionSet BlockedOptionalExtensionSet;
 	{
@@ -788,13 +801,9 @@ bool FOpenXRHMDModule::InitInstance()
 			}
 		}
 		// Logging to show which extension plugins are blocking what optional extensions.
-		UE_LOG(LogHMD, Log, TEXT("BlockedOptionalExtensions:"));
-		if (BlockedOptionalExtensionMap.IsEmpty())
-		{
-			UE_LOG(LogHMD, Log, TEXT("  None"));
-		}
-		else
+		if (!BlockedOptionalExtensionMap.IsEmpty())
 		{
+			UE_LOG(LogHMD, Log, TEXT("BlockedOptionalExtensionSet:"));
 			for (TPair< IOpenXRExtensionPlugin*, TArray<const ANSICHAR*> > Pair : BlockedOptionalExtensionMap)
 			{
 				UE_LOG(LogHMD, Log, TEXT("  %s is blocking:"), *(Pair.Key->GetDisplayName()));
@@ -818,7 +827,7 @@ bool FOpenXRHMDModule::InitInstance()
 			UE_LOG(LogHMD, Error, TEXT("Could not get required OpenXR extensions."));
 			return false;
 		}
-		OpenXRHMDModule::LogExtensions(TEXT("  RequiredExtensions"), RequiredExtensions);
+		OpenXRHMDModule::LogExtensionsVerbose(true, TEXT("  RequiredExtensions"), RequiredExtensions);
 
 		// Query optional extensions
 		OptionalExtensions.Empty();
@@ -827,12 +836,9 @@ bool FOpenXRHMDModule::InitInstance()
 			UE_LOG(LogHMD, Error, TEXT("Could not get optional OpenXR extensions."));
 			return false;
 		}
-		OpenXRHMDModule::LogExtensions(TEXT("  OptionalExtensions"), OptionalExtensions);
-
-		OpenXRHMDModule::BlockOptionalExtensions(OptionalExtensions, BlockedOptionalExtensionSet);
-		OpenXRHMDModule::LogExtensions(TEXT("  OptionalExtensions After Blocking:"), OptionalExtensions);
+		OpenXRHMDModule::LogExtensionsVerbose(true, TEXT("  OptionalExtensions"), OptionalExtensions);
 
-		if (!EnableExtensions(RequiredExtensions, OptionalExtensions, Extensions))
+		if (!OpenXRHMDModule::EnableExtensions(AvailableExtensions, RequiredExtensions, OptionalExtensions, Extensions, FullyBlockedExtensionsSet, BlockedOptionalExtensionSet))
 		{
 			UE_LOG(LogHMD, Error, TEXT("Could not enable all required OpenXR extensions."));
 			return false;
@@ -840,11 +846,6 @@ bool FOpenXRHMDModule::InitInstance()
 		ExtensionSet.Append(Extensions);
 	}
 
-	if (AvailableExtensions.Contains(XR_EPIC_VIEW_CONFIGURATION_FOV_EXTENSION_NAME))
-	{
-		ExtensionSet.Add(XR_EPIC_VIEW_CONFIGURATION_FOV_EXTENSION_NAME);
-	}
-
 	UE_LOG(LogHMD, Log, TEXT("Extensions for IOpenXRExtensionPlugins:"));
 	ExtensionPlugins.Reset();
 	for (IOpenXRExtensionPlugin* Plugin : ExtModules)
@@ -859,7 +860,7 @@ bool FOpenXRHMDModule::InitInstance()
 			// Ignore the plugin if the query fails
 			continue;
 		}
-		OpenXRHMDModule::LogExtensions(!RequiredExtensions.IsEmpty(), TEXT("    RequiredExtensions"), RequiredExtensions);
+		OpenXRHMDModule::LogExtensionsVerbose(!RequiredExtensions.IsEmpty(), TEXT("    RequiredExtensions"), RequiredExtensions);
 
 		// Query optional extensions
 		OptionalExtensions.Empty();
@@ -868,13 +869,9 @@ bool FOpenXRHMDModule::InitInstance()
 			// Ignore the plugin if the query fails
 			continue;
 		}
-		const bool bLogOptionalExtensions = !OptionalExtensions.IsEmpty();
-		OpenXRHMDModule::LogExtensions(bLogOptionalExtensions, TEXT("    OptionalExtensions"), OptionalExtensions);
-
-		OpenXRHMDModule::BlockOptionalExtensions(OptionalExtensions, BlockedOptionalExtensionSet);
-		OpenXRHMDModule::LogExtensions(bLogOptionalExtensions, TEXT("    OptionalExtensions After Blocking:"), OptionalExtensions);
+		OpenXRHMDModule::LogExtensionsVerbose(!OptionalExtensions.IsEmpty(), TEXT("    OptionalExtensions"), OptionalExtensions);
 
-		if (!EnableExtensions(RequiredExtensions, OptionalExtensions, Extensions))
+		if (!OpenXRHMDModule::EnableExtensions(AvailableExtensions, RequiredExtensions, OptionalExtensions, Extensions, FullyBlockedExtensionsSet, BlockedOptionalExtensionSet))
 		{
 			// Ignore the plugin if the required extension could not be enabled
 			FString ModuleName = Plugin->GetDisplayName();
@@ -898,7 +895,7 @@ bool FOpenXRHMDModule::InitInstance()
 		EnabledExtensions.Add(Ext);
 	}
 	
-	OpenXRHMDModule::LogExtensions(TEXT("    EnabledExtensions:"), EnabledExtensions);
+	OpenXRHMDModule::LogExtensions(true, TEXT("    EnabledExtensions:"), EnabledExtensions);
 
 	// Enable layers, if specified by CVar.
 	// Note: For the validation layer to work on Windows (as of latest OpenXR runtime, August 2019), the following are required:
@@ -1180,7 +1177,7 @@ bool FOpenXRHMDModule::TryCreateInstance(XrInstanceCreateInfo& Info)
 				int32 NumProblematicExtensions = ProblematicExtensionList.Num();
 				check(NumProblematicExtensions < 64);
 				check(NumProblematicExtensions <= 3);
-				NumProblematicExtensions = FMath::Max(NumProblematicExtensions, 3);
+				NumProblematicExtensions = FMath::Min(NumProblematicExtensions, 3);
 				uint64 Bitfield = (1ull << (NumProblematicExtensions)) - 1;  // Get a 1 for each extension.
 				do
 				{
diff --git a/Engine/Plugins/Runtime/OpenXR/Source/OpenXRHMD/Private/OpenXRHMDModule.h b/Engine/Plugins/Runtime/OpenXR/Source/OpenXRHMD/Private/OpenXRHMDModule.h
--- a/Engine/Plugins/Runtime/OpenXR/Source/OpenXRHMD/Private/OpenXRHMDModule.h
+++ b/Engine/Plugins/Runtime/OpenXR/Source/OpenXRHMD/Private/OpenXRHMDModule.h
@@ -79,7 +79,6 @@ private:
 	bool InitInstance();
 	bool TryCreateInstance(XrInstanceCreateInfo& Info);
 	PFN_xrGetInstanceProcAddr GetDefaultLoader();
-	bool EnableExtensions(const TArray<const ANSICHAR*>& RequiredExtensions, const TArray<const ANSICHAR*>& OptionalExtensions, TArray<const ANSICHAR*>& OutExtensions);
 	bool GetRequiredExtensions(TArray<const ANSICHAR*>& OutExtensions);
 	bool GetOptionalExtensions(TArray<const ANSICHAR*>& OutExtensions);
 };
diff --git a/Engine/Plugins/Runtime/OpenXR/Source/OpenXRHMD/Public/OpenXRHMDSettings.h b/Engine/Plugins/Runtime/OpenXR/Source/OpenXRHMD/Public/OpenXRHMDSettings.h
--- a/Engine/Plugins/Runtime/OpenXR/Source/OpenXRHMD/Public/OpenXRHMDSettings.h
+++ b/Engine/Plugins/Runtime/OpenXR/Source/OpenXRHMD/Public/OpenXRHMDSettings.h
@@ -40,6 +40,12 @@ public:
 		DisplayName = "Enable OpenXR 1.1"))
 	bool bIsOpenXR1_1Enabled = true;
 
+	/** Extension Blocking */
+	UPROPERTY(config, EditAnywhere, Category = "OpenXR", meta = (
+		ToolTip = "List of OpenXR extension names that will be blocked and not used.\n The names must exactly match those provided in openxr_platform.h (eg \"XR_KHR_vulkan_enable\" from #define XR_KHR_VULKAN_ENABLE_EXTENSION_NAME \"XR_KHR_vulkan_enable\")\nIf a UE Core optional extension is blocked UE's OpenXR support should continue to function, but functionality or performance may be degraded.\nIf a UE Core required extension is blocked UE's OpenXR support will be disabled.\nIf a OpenXRExtensionPlugin optional extension is blocked the plugin should continue to function, but functionality or performance may be degraded.\nIf a OpenXRExtensionPlugin required extension is blocked that plugin will be disabled.\nAnticipated use cases:\n  Facilitation of testing without certain extensions\n  Working around runtime bugs where an extension is malfunctioning (Note that if the runtime is later fixed your project will still not use it!).\n You will need to restart the editor after changing this setting to have it effect PIE.",
+		DisplayName = "Blocked Extensions"))
+	TArray<FString> BlockedExtensions;
+
 public:
 	virtual FName GetCategoryName() const override { return FName(TEXT("Plugins")); }
 };
