Add support for stereo layer underlays (layers which are rendered below the projection layer/normal UE scene).
Add OpenXR project settings/CVars to control projection layer priority and whether or not face-locked layers should render above all other layers (defaults to false). Legacy behavior can be recreated by setting xr.OpenXRSortFaceLockedLayersAboveOtherLayers=1 and xr.OpenXRProjectionLayerPriority=INT_MIN.
Some OpenXR runtimes do not support underlays and will render all stereo layers above the projection layer regardless of what order they're submitted in.

#Synced-CL 52062538

diff --git a/Engine/Plugins/Runtime/OpenXR/Source/OpenXRHMD/Private/OpenXRHMD.cpp b/Engine/Plugins/Runtime/OpenXR/Source/OpenXRHMD/Private/OpenXRHMD.cpp
--- a/Engine/Plugins/Runtime/OpenXR/Source/OpenXRHMD/Private/OpenXRHMD.cpp
+++ b/Engine/Plugins/Runtime/OpenXR/Source/OpenXRHMD/Private/OpenXRHMD.cpp
@@ -173,6 +173,18 @@ static TAutoConsoleVariable<bool> CVarOpenXRFrameSynthesis(
 	TEXT("Because normal velocity rendering is disabled when r.Velocity.DirectlyRenderOpenXRMotionVectors=True, temporal anti-aliasing and motion blur will be automatically disabled."),
 	ECVF_RenderThreadSafe);
 
+static TAutoConsoleVariable<int32> CVarOpenXRProjectionLayerPriority(
+	TEXT("xr.OpenXRProjectionLayerPriority"),
+	0,
+	TEXT("Render priority for the projection layer (main UE environment). Stereo layers with priorities greater than or equal to this value will render above the projection layer. Layers with priorities less than this value will render below the projection layer."),
+	ECVF_RenderThreadSafe);
+
+static TAutoConsoleVariable<bool> CVarOpenXRSortFaceLockedLayersAboveOtherLayers(
+	TEXT("xr.OpenXRSortFaceLockedLayersAboveOtherLayers"),
+	false,
+	TEXT("Sort face locked stereo layers above all other layers, including the projection layer, even if the face locked layer has a lower priority. This was the standard behavior before UE 5.8."),
+	ECVF_RenderThreadSafe);
+
 namespace {
 	static TSet<XrViewConfigurationType> SupportedViewConfigurations{ XR_VIEW_CONFIGURATION_TYPE_PRIMARY_STEREO, XR_VIEW_CONFIGURATION_TYPE_PRIMARY_STEREO_WITH_FOVEATED_INSET };
 
@@ -1955,7 +1967,7 @@ void FOpenXRHMD::EnumerateViews(FPipelinedFrameState& PipelineState)
 		// Variable size between views is not currently supported
 		XrFrameSynthesisConfigViewEXT& FrameSynthConfig = ViewFrameSynthesis[0];
 		RecommendedMotionVectorTextureSize = FIntPoint(FrameSynthConfig.recommendedMotionVectorImageRectWidth, FrameSynthConfig.recommendedMotionVectorImageRectHeight);
-		check(RecommendedMotionVectorTextureSize.X && RecommendedMotionVectorTextureSize.Y);
+		ensure(RecommendedMotionVectorTextureSize.X && RecommendedMotionVectorTextureSize.Y);
 	}
 }
 
@@ -2455,7 +2467,7 @@ void FOpenXRHMD::DestroySession()
 
 		PipelinedLayerStateRendering.ColorSwapchain.Reset();
 		PipelinedLayerStateRendering.DepthSwapchain.Reset();
-		PipelinedLayerStateRendering.NativeOverlaySwapchains.Reset();
+		PipelinedLayerStateRendering.NativeLayerSwapchains.Reset();
 		PipelinedLayerStateRendering.EmulatedLayerState.EmulationSwapchain.Reset();
 
 		// TODO: Once we handle OnFinishRendering_RHIThread + StopSession interactions
@@ -2463,7 +2475,7 @@ void FOpenXRHMD::DestroySession()
 		// `ensure` here to make sure these are released.
 		PipelinedLayerStateRHI.ColorSwapchain.Reset();
 		PipelinedLayerStateRHI.DepthSwapchain.Reset();
-		PipelinedLayerStateRHI.NativeOverlaySwapchains.Reset();
+		PipelinedLayerStateRHI.NativeLayerSwapchains.Reset();
 		PipelinedLayerStateRHI.EmulatedLayerState.EmulationSwapchain.Reset();
 
 		PipelinedFrameStateGame.TrackingSpace.Reset();
@@ -3240,17 +3252,23 @@ void FOpenXRHMD::SetupFrameLayers_GameThread()
 			}
 		});
 
-		auto LayerCompare = [](const auto& DescA, const auto& DescB) -> bool
+		const bool bUseLegacyFaceLockedLayersSorting = CVarOpenXRSortFaceLockedLayersAboveOtherLayers.GetValueOnAnyThread();
+
+		// Sort from lowest priority to highest priority. Layers with higher priority will render on top of layers with lower priority.
+		auto LayerCompare = [bUseLegacyFaceLockedLayersSorting](const auto& DescA, const auto& DescB) -> bool
 		{
-			if ((DescA.PositionType == IStereoLayers::FaceLocked) != (DescB.PositionType == IStereoLayers::FaceLocked))
+			if (bUseLegacyFaceLockedLayersSorting && (DescA.PositionType == IStereoLayers::FaceLocked) != (DescB.PositionType == IStereoLayers::FaceLocked))
 			{
 				return DescB.PositionType == IStereoLayers::FaceLocked;
 			}
-			if (DescA.Priority != DescB.Priority)
+			else if (DescA.Priority != DescB.Priority)
 			{
 				return DescA.Priority < DescB.Priority;
 			}
-			return DescA.Id < DescB.Id;
+			else
+			{
+				return DescA.Id < DescB.Id;
+			}
 		};
 
 		LayersToRender.Sort(LayerCompare);
@@ -3401,16 +3419,20 @@ void FOpenXRHMD::SetupFrameLayers_RenderThread(FRDGBuilder& GraphBuilder)
 	const float WorldToMeters = GetWorldToMetersScale();
 
 	PipelinedLayerStateRendering.NativeOverlays.Reset(NativeLayers.Num());
-	PipelinedLayerStateRendering.NativeOverlaySwapchains.Reset(NativeLayers.Num());
+	PipelinedLayerStateRendering.NativeUnderlays.Reset(NativeLayers.Num());
+	PipelinedLayerStateRendering.NativeLayerSwapchains.Reset(NativeLayers.Num());
 	PipelinedLayerStateRendering.CompositionDepthTestLayers.Reset(NativeLayers.Num());
 
+	const bool bUseLegacyFaceLockedLayersSorting = CVarOpenXRSortFaceLockedLayersAboveOtherLayers.GetValueOnAnyThread();
+	const int32 ProjectionLayerPriority = CVarOpenXRProjectionLayerPriority.GetValueOnAnyThread();
+
 	// Set up our OpenXR info per native layer. Emulated layers have everything in FLayerDesc.
 	for (FOpenXRLayer& Layer : NativeLayers)
 	{
 		FReadScopeLock DeviceLock(DeviceMutex);
 
-		XrSpace Space = Layer.Desc.PositionType == ELayerType::FaceLocked ?
-			DeviceSpaces[HMDDeviceId].Space : PipelinedFrameStateRendering.TrackingSpace->Handle;
+		const bool bIsFaceLocked = Layer.Desc.PositionType == ELayerType::FaceLocked;
+		XrSpace Space = bIsFaceLocked ? DeviceSpaces[HMDDeviceId].Space : PipelinedFrameStateRendering.TrackingSpace->Handle;
 		
 		EOpenXRLayerCreationFlags LayerCreationFlags = EOpenXRLayerCreationFlags::None;
 		LayerCreationFlags |= IsExtensionEnabled(XR_KHR_COMPOSITION_LAYER_EQUIRECT2_EXTENSION_NAME) ? EOpenXRLayerCreationFlags::EquirectLayer2Supported : EOpenXRLayerCreationFlags::None;
@@ -3419,7 +3441,15 @@ void FOpenXRHMD::SetupFrameLayers_RenderThread(FRDGBuilder& GraphBuilder)
 		TArray<FXrCompositionLayerUnion> Headers = Layer.CreateOpenXRLayer(InvTrackingToWorld, WorldToMeters, Space, LayerCreationFlags);
 		Layer.ApplyCompositionDepthTestLayer(Headers, LayerCreationFlags, PipelinedLayerStateRendering.CompositionDepthTestLayers);
 
-		PipelinedLayerStateRendering.NativeOverlays.Append(Headers);
+		if (Layer.Desc.Priority >= ProjectionLayerPriority || (bUseLegacyFaceLockedLayersSorting && bIsFaceLocked))
+		{
+			PipelinedLayerStateRendering.NativeOverlays.Append(Headers);
+		}
+		else
+		{
+			PipelinedLayerStateRendering.NativeUnderlays.Append(Headers);
+		}
+		
 		UpdateLayerSwapchainTexture(Layer, GraphBuilder);
 	}
 }
@@ -3443,7 +3473,7 @@ void FOpenXRHMD::UpdateLayerSwapchainTexture(FOpenXRLayer& Layer, FRDGBuilder& G
 			CopySwapchainTexture_RenderThread(GraphBuilder, SrcTextureRDG, FIntRect(), Layer.RightEye.Swapchain, DstRect, false, SrcTextureCopyModifier, FeatureLevel, ShaderPlatform);
 			Layer.RightEye.bUpdateTexture = !Layer.RightEye.bStaticSwapchain;
 		}
-		PipelinedLayerStateRendering.NativeOverlaySwapchains.Add(Layer.RightEye.Swapchain);
+		PipelinedLayerStateRendering.NativeLayerSwapchains.Add(Layer.RightEye.Swapchain);
 	}
 	if (Layer.LeftEye.Swapchain.IsValid() && Layer.LeftEye.Texture.IsValid())
 	{
@@ -3455,7 +3485,7 @@ void FOpenXRHMD::UpdateLayerSwapchainTexture(FOpenXRLayer& Layer, FRDGBuilder& G
 			CopySwapchainTexture_RenderThread(GraphBuilder, SrcTextureRDG, FIntRect(), Layer.LeftEye.Swapchain, DstRect, false, SrcTextureCopyModifier, FeatureLevel, ShaderPlatform);
 			Layer.LeftEye.bUpdateTexture = !Layer.LeftEye.bStaticSwapchain;
 		}
-		PipelinedLayerStateRendering.NativeOverlaySwapchains.Add(Layer.LeftEye.Swapchain);
+		PipelinedLayerStateRendering.NativeLayerSwapchains.Add(Layer.LeftEye.Swapchain);
 	}
 }
 
@@ -4390,61 +4420,69 @@ void FOpenXRHMD::OnFinishRendering(IRHICommandContext* RHICmdContext)
 	if (bIsRunning)
 	{
 		TArray<XrCompositionLayerBaseHeader*> Headers;
-		XrCompositionLayerProjection Layer = {};
+
+		// Layers are rendered in the order they're submitted to the compositor, so earlier layers in the Headers list will appear below later layers.
+		// So we add underlays, then projection layers, then overlays. Underlays will appear below the projection layers, and overlays will appear above.
+		for (FXrCompositionLayerUnion& NativeUnderlay : PipelinedLayerStateRHI.NativeUnderlays)
+		{
+			Headers.Add(reinterpret_cast<XrCompositionLayerBaseHeader*>(&NativeUnderlay.Header));
+		}
+
+		XrCompositionLayerProjection BackgroundLayer = {};
 		XrCompositionLayerAlphaBlendFB LayerAlphaBlend = { XR_TYPE_COMPOSITION_LAYER_ALPHA_BLEND_FB };
 		XrCompositionLayerColorScaleBiasKHR ColorScaleBias = { XR_TYPE_COMPOSITION_LAYER_COLOR_SCALE_BIAS_KHR };
 		XrCompositionLayerDepthTestFB LayerDepthTest = { XR_TYPE_COMPOSITION_LAYER_DEPTH_TEST_FB };
 		if (EnumHasAnyFlags(PipelinedLayerStateRHI.LayerStateFlags, EOpenXRLayerStateFlags::SubmitBackgroundLayer))
 		{
-			Layer.type = XR_TYPE_COMPOSITION_LAYER_PROJECTION;
-			Layer.next = nullptr;
-			Layer.layerFlags = bProjectionLayerAlphaEnabled ? XR_COMPOSITION_LAYER_BLEND_TEXTURE_SOURCE_ALPHA_BIT : 0;
-			Layer.space = PipelinedFrameStateRHI.TrackingSpace->Handle;
-			Layer.viewCount = PipelinedLayerStateRHI.ProjectionLayers.Num();
-			Layer.views = PipelinedLayerStateRHI.ProjectionLayers.GetData();
+			BackgroundLayer.type = XR_TYPE_COMPOSITION_LAYER_PROJECTION;
+			BackgroundLayer.next = nullptr;
+			BackgroundLayer.layerFlags = bProjectionLayerAlphaEnabled ? XR_COMPOSITION_LAYER_BLEND_TEXTURE_SOURCE_ALPHA_BIT : 0;
+			BackgroundLayer.space = PipelinedFrameStateRHI.TrackingSpace->Handle;
+			BackgroundLayer.viewCount = PipelinedLayerStateRHI.ProjectionLayers.Num();
+			BackgroundLayer.views = PipelinedLayerStateRHI.ProjectionLayers.GetData();
 
-			Headers.Add(reinterpret_cast<XrCompositionLayerBaseHeader*>(&Layer));
+			Headers.Add(reinterpret_cast<XrCompositionLayerBaseHeader*>(&BackgroundLayer));
 
 			if (CVarOpenXRInvertAlpha.GetValueOnRenderThread())
 			{
 				// These two extensions do the same thing, the first one is more modern.  The second is to keep older runtimes functioning.
 				if(IsExtensionEnabled(XR_EXT_COMPOSITION_LAYER_INVERTED_ALPHA_EXTENSION_NAME))
 				{
-					Layer.layerFlags |= XR_COMPOSITION_LAYER_INVERTED_ALPHA_BIT_EXT;
+					BackgroundLayer.layerFlags |= XR_COMPOSITION_LAYER_INVERTED_ALPHA_BIT_EXT;
 				}
 				else if (IsExtensionEnabled(XR_FB_COMPOSITION_LAYER_ALPHA_BLEND_EXTENSION_NAME))
 				{
-					LayerAlphaBlend.next = const_cast<void*>(Layer.next);
+					LayerAlphaBlend.next = const_cast<void*>(BackgroundLayer.next);
 					LayerAlphaBlend.srcFactorColor = PipelinedLayerStateRHI.BasePassLayerBlendParams.srcFactorColor;
 					LayerAlphaBlend.dstFactorColor = PipelinedLayerStateRHI.BasePassLayerBlendParams.dstFactorColor;
 					LayerAlphaBlend.srcFactorAlpha = PipelinedLayerStateRHI.BasePassLayerBlendParams.srcFactorAlpha;
 					LayerAlphaBlend.dstFactorAlpha = PipelinedLayerStateRHI.BasePassLayerBlendParams.dstFactorAlpha;
 
-					Layer.next = &LayerAlphaBlend;
+					BackgroundLayer.next = &LayerAlphaBlend;
 				}
 			}
 
 			if (bCompositionLayerColorScaleBiasSupported)
 			{
-				ColorScaleBias.next = const_cast<void*>(Layer.next);
+				ColorScaleBias.next = const_cast<void*>(BackgroundLayer.next);
 				ColorScaleBias.colorScale = PipelinedLayerStateRHI.LayerColorScaleAndBias.ColorScale;
 				ColorScaleBias.colorBias = PipelinedLayerStateRHI.LayerColorScaleAndBias.ColorBias;
 
-				Layer.next = &ColorScaleBias;
+				BackgroundLayer.next = &ColorScaleBias;
 			}
 
 			if (IsExtensionEnabled(XR_FB_COMPOSITION_LAYER_DEPTH_TEST_EXTENSION_NAME) && 
 				EnumHasAnyFlags(PipelinedLayerStateRendering.LayerStateFlags, EOpenXRLayerStateFlags::SubmitDepthLayer))
 			{
-				LayerDepthTest.next = const_cast<void*>(Layer.next);
+				LayerDepthTest.next = const_cast<void*>(BackgroundLayer.next);
 				LayerDepthTest.depthMask = true;
 				LayerDepthTest.compareOp = XR_COMPARE_OP_LESS_FB;
-				Layer.next = &LayerDepthTest;
+				BackgroundLayer.next = &LayerDepthTest;
 			}
 
 			for (IOpenXRExtensionPlugin* Module : ExtensionPlugins)
 			{
-				Layer.next = Module->OnEndProjectionLayer_RHIThread(Session, 0, Layer.next, Layer.layerFlags);
+				BackgroundLayer.next = Module->OnEndProjectionLayer_RHIThread(Session, 0, BackgroundLayer.next, BackgroundLayer.layerFlags);
 			}
 		}
 		
@@ -4464,7 +4502,15 @@ void FOpenXRHMD::OnFinishRendering(IRHICommandContext* RHICmdContext)
 			Headers.Add(reinterpret_cast<XrCompositionLayerBaseHeader*>(&CompositedLayer));
 		}
 
-		AddLayersToHeaders(Headers);
+		for (FXrCompositionLayerUnion& NativeOverlay : PipelinedLayerStateRHI.NativeOverlays)
+		{
+			Headers.Add(reinterpret_cast<XrCompositionLayerBaseHeader*>(&NativeOverlay.Header));
+		}
+
+		for (IOpenXRExtensionPlugin* Module : ExtensionPlugins)
+		{
+			Module->UpdateCompositionLayers_RHIThread(Session, Headers);
+		}
 
 		void* Next = nullptr;
 		XrRHIContextEPIC RHIContextEPIC = { (XrStructureType)XR_TYPE_RHI_CONTEXT_EPIC };
@@ -4509,19 +4555,6 @@ void FOpenXRHMD::OnFinishRendering(IRHICommandContext* RHICmdContext)
 	bIsRendering = false;
 }
 
-void FOpenXRHMD::AddLayersToHeaders(TArray<XrCompositionLayerBaseHeader*>& Headers)
-{
-	for (FXrCompositionLayerUnion& Layer : PipelinedLayerStateRHI.NativeOverlays)
-	{
-		Headers.Add(reinterpret_cast<XrCompositionLayerBaseHeader*>(&Layer.Header));
-	}
-
-	for (IOpenXRExtensionPlugin* Module : ExtensionPlugins)
-	{
-		Module->UpdateCompositionLayers_RHIThread(Session, Headers);
-	}
-}
-
 bool FOpenXRHMD::HDRGetMetaDataForStereo(EDisplayOutputFormat& OutDisplayOutputFormat, EDisplayColorGamut& OutDisplayColorGamut, bool& OutbHDRSupported)
 {
 	if (RenderBridge == nullptr)
diff --git a/Engine/Plugins/Runtime/OpenXR/Source/OpenXRHMD/Private/OpenXRHMD.h b/Engine/Plugins/Runtime/OpenXR/Source/OpenXRHMD/Private/OpenXRHMD.h
--- a/Engine/Plugins/Runtime/OpenXR/Source/OpenXRHMD/Private/OpenXRHMD.h
+++ b/Engine/Plugins/Runtime/OpenXR/Source/OpenXRHMD/Private/OpenXRHMD.h
@@ -205,6 +205,7 @@ public:
 	struct FPipelinedLayerState
 	{
 		TArray<FXrCompositionLayerUnion> NativeOverlays;
+		TArray<FXrCompositionLayerUnion> NativeUnderlays; // Underlays will be rendered below the projection layers
 		TArray<XrCompositionLayerProjectionView> ProjectionLayers;
 		TArray<XrCompositionLayerDepthInfoKHR> DepthLayers;
 		TArray<XrCompositionLayerDepthTestFB> CompositionDepthTestLayers;
@@ -222,7 +223,7 @@ public:
 		FXRSwapChainPtr DepthSwapchain;
 		FXRSwapChainPtr MotionVectorSwapchain;
 		FXRSwapChainPtr MotionVectorDepthSwapchain;
-		TArray<FXRSwapChainPtr> NativeOverlaySwapchains;
+		TArray<FXRSwapChainPtr> NativeLayerSwapchains;	// This is currently only used as a reference to keep the swapchains from being destroyed
 
 		FEmulatedLayerState EmulatedLayerState;
 
@@ -531,7 +532,6 @@ private:
 
 	void UpdateLayerSwapchainTexture(FOpenXRLayer& Layer, FRDGBuilder& GraphBuilder);
 	void ConfigureLayerSwapchains(const struct FLayerToUpdateSwapchain& Update, FOpenXRLayer& Layer);
-	void AddLayersToHeaders(TArray<XrCompositionLayerBaseHeader*>& Headers);
 
 	// Gets the user path corresponding to the input EControllerHand value
 	static XrPath GetUserPathForControllerHand(EControllerHand Hand);
diff --git a/Engine/Plugins/Runtime/OpenXR/Source/OpenXRHMD/Public/OpenXRHMDSettings.h b/Engine/Plugins/Runtime/OpenXR/Source/OpenXRHMD/Public/OpenXRHMDSettings.h
--- a/Engine/Plugins/Runtime/OpenXR/Source/OpenXRHMD/Public/OpenXRHMDSettings.h
+++ b/Engine/Plugins/Runtime/OpenXR/Source/OpenXRHMD/Public/OpenXRHMDSettings.h
@@ -9,7 +9,7 @@
 /**
 * Implements the settings for the OpenXR plugin.
 */
-UCLASS(config=Engine, defaultconfig, meta = (DisplayName = "OpenXR Settings"))
+UCLASS(config=Engine, defaultconfig, meta = (DisplayName = "OpenXR"))
 class OPENXRHMD_API UOpenXRHMDSettings : public UDeveloperSettingsBackedByCVars
 {
 	GENERATED_BODY()
@@ -46,6 +46,20 @@ public:
 		DisplayName = "Blocked Extensions"))
 	TArray<FString> BlockedExtensions;
 
+	/** Projection Layer Priority */
+	UPROPERTY(config, EditAnywhere, Category = "Stereo Layers", meta = (
+		ConsoleVariable = "xr.OpenXRProjectionLayerPriority",
+		ToolTip = "Render priority for the projection layer (main UE environment).\nStereo layers with priorities greater than or equal to this value will render above the projection layer. Layers with priorities less than this value will render below the projection layer.\nNote that some OpenXR runtimes will ignore this setting and always render the projection layer below all stereo layers.",
+		DisplayName = "Projection Layer Priority"))
+	int32 Priority = 0;
+
+	/** Sort face-locked layers above all other layers */
+	UPROPERTY(config, EditAnywhere, Category = "Stereo Layers", meta = (
+		ConsoleVariable = "xr.OpenXRSortFaceLockedLayersAboveOtherLayers",
+		ToolTip = "Sort face locked stereo layers above all other layers, including the projection layer, even if the face locked layer has a lower priority.",
+		DisplayName = "Sort Face Locked Layers Above Other Layers"))
+	bool bSortFaceLockedLayersAboveOtherLayers = false;
+
 public:
 	virtual FName GetCategoryName() const override { return FName(TEXT("Plugins")); }
 };
diff --git a/Engine/Source/Runtime/Engine/Classes/Components/StereoLayerComponent.h b/Engine/Source/Runtime/Engine/Classes/Components/StereoLayerComponent.h
--- a/Engine/Source/Runtime/Engine/Classes/Components/StereoLayerComponent.h
+++ b/Engine/Source/Runtime/Engine/Classes/Components/StereoLayerComponent.h
@@ -405,7 +405,7 @@ protected:
 	UPROPERTY(EditAnywhere, BlueprintReadOnly, NoClear, Instanced, Category = "StereoLayer", DisplayName="Stereo Layer Shape")
 	TObjectPtr<UStereoLayerShape> Shape;
 
-	/** Render priority among all stereo layers, higher priority render on top of lower priority **/
+	/** Render priority among all stereo layers. Layers with higher priorities will be rendered on top of layers with lower priorities. **/
 	UPROPERTY(EditAnywhere, BlueprintReadOnly, export, Category= "StereoLayer")
 	int32 Priority;
 
