Validate swapchain pointers before enqueuing acquire lambdas in AllocateRenderTargetTextures
Remove incorrect assertion on motion vector swapchain index parity in AcquireDepthTexture; the OpenXR runtime may return different indices per swapchain and parity is not required for submission.

#Synced-CL 52549286

diff --git a/Engine/Plugins/Runtime/OpenXR/Source/OpenXRHMD/Private/OpenXRHMD.cpp b/Engine/Plugins/Runtime/OpenXR/Source/OpenXRHMD/Private/OpenXRHMD.cpp
--- a/Engine/Plugins/Runtime/OpenXR/Source/OpenXRHMD/Private/OpenXRHMD.cpp
+++ b/Engine/Plugins/Runtime/OpenXR/Source/OpenXRHMD/Private/OpenXRHMD.cpp
@@ -2785,11 +2785,6 @@ int32 FOpenXRHMD::AcquireDepthTexture()
 				MotionVectorSwapchain->IncrementSwapChainIndex_RHIThread();
 				MotionVectorDepthSwapchain->IncrementSwapChainIndex_RHIThread();
 			}
-
-			const int32 MotionVectorIndex = MotionVectorSwapchain->GetSwapChainIndex_RHIThread();
-			const int32 MotionVectorDepthIndex = MotionVectorDepthSwapchain->GetSwapChainIndex_RHIThread();
-			
-			check(MotionVectorIndex == MotionVectorDepthIndex);
 		}
 
 		if (DepthSwapchain)
@@ -2853,6 +2848,11 @@ bool FOpenXRHMD::AllocateRenderTargetTextures(FRHICommandListBase& RHICmdList, u
 	}
 
 	FXRSwapChainPtr& Swapchain = PipelinedLayerStateRendering.ColorSwapchain;
+	if (!Swapchain.IsValid())
+	{
+		return false;
+	}
+
 	uint8 ActualFormat = Format;
 	ensureMsgf(NumSamples == 1, TEXT("OpenXR supports MSAA swapchains, but engine logic expects the swapchain target to be 1x."));
 	
@@ -2869,7 +2869,9 @@ bool FOpenXRHMD::AllocateRenderTargetTextures(FRHICommandListBase& RHICmdList, u
 		FClearValueBinding::Transparent,
 		AuxiliaryCreateFlags
 	};
+
 	bool bAllocatedColorTextures = AllocateSwapchainTextures_RenderThread(ColorSwapchainProperties, Swapchain, ActualFormat);
+
 	// Image will be acquired by the renderer if supported, if not we acquire it ahead of time here
 	if (!bIsAcquireOnAnyThreadSupported)
 	{
@@ -2882,16 +2884,10 @@ bool FOpenXRHMD::AllocateRenderTargetTextures(FRHICommandListBase& RHICmdList, u
 		});
 		RHICmdList.GetAsImmediate().ImmediateFlush(EImmediateFlushType::FlushRHIThread);
 	}
-	if (Swapchain)
-	{
-		if (FBFoveationImageGenerator && FBFoveationImageGenerator->IsFoveationExtensionEnabled())
-		{
-			FBFoveationImageGenerator->UpdateFoveationImages(/* bReallocatedSwapchain */ bAllocatedColorTextures);
-		}
-	}
-	else
+
+	if (FBFoveationImageGenerator && FBFoveationImageGenerator->IsFoveationExtensionEnabled())
 	{
-		return false;
+		FBFoveationImageGenerator->UpdateFoveationImages(/* bReallocatedSwapchain */ bAllocatedColorTextures);
 	}
 
 	// Grab the presentation texture out of the swapchain.
@@ -2922,17 +2918,17 @@ bool FOpenXRHMD::AllocateRenderTargetTextures(FRHICommandListBase& RHICmdList, u
 		AllocateSwapchainTextures_RenderThread(EmulationSwapchainProperties, PipelinedLayerStateRendering.EmulatedLayerState.EmulationSwapchain, UnusedActualFormat);
 
 		// Image will be acquired by the renderer if supported, if not we acquire it ahead of time here
-		if (!bIsAcquireOnAnyThreadSupported)
+		FXRSwapChainPtr& EmulationSwapchain = PipelinedLayerStateRendering.EmulatedLayerState.EmulationSwapchain;
+		if (!bIsAcquireOnAnyThreadSupported && EmulationSwapchain.IsValid())
 		{
-			FXRSwapChainPtr& EmulationSwapchain = PipelinedLayerStateRendering.EmulatedLayerState.EmulationSwapchain;
-				RHICmdList.EnqueueLambda([this, EmulationSwapchain](FRHICommandListBase&)
+			RHICmdList.EnqueueLambda([this, EmulationSwapchain](FRHICommandListBase&)
 			{
 				RenderBridge->RunOnRHISubmissionThread([EmulationSwapchain]
 				{
 					EmulationSwapchain->IncrementSwapChainIndex_RHIThread();
 				});
 			});
-				RHICmdList.GetAsImmediate().ImmediateFlush(EImmediateFlushType::FlushRHIThread);
+			RHICmdList.GetAsImmediate().ImmediateFlush(EImmediateFlushType::FlushRHIThread);
 		}
 	}
 
@@ -2975,10 +2971,10 @@ bool FOpenXRHMD::AllocateRenderTargetTextures(FRHICommandListBase& RHICmdList, u
 		AllocateSwapchainTextures_RenderThread(MotionVectorDepthSwapchainProperties, PipelinedLayerStateRendering.MotionVectorDepthSwapchain, UnusedActualFormat);
 
 		// Images will be acquired by the renderer if supported, if not we acquire it ahead of time here
-		if (!bIsAcquireOnAnyThreadSupported)
+		FXRSwapChainPtr& MotionVectorSwapchain = PipelinedLayerStateRendering.MotionVectorSwapchain;
+		FXRSwapChainPtr& MotionVectorDepthSwapchain = PipelinedLayerStateRendering.MotionVectorDepthSwapchain;
+		if (!bIsAcquireOnAnyThreadSupported && MotionVectorSwapchain.IsValid() && MotionVectorDepthSwapchain.IsValid())
 		{
-			FXRSwapChainPtr& MotionVectorSwapchain = PipelinedLayerStateRendering.MotionVectorSwapchain;
-			FXRSwapChainPtr& MotionVectorDepthSwapchain = PipelinedLayerStateRendering.MotionVectorDepthSwapchain;
 			RHICmdList.EnqueueLambda([this, MotionVectorSwapchain, MotionVectorDepthSwapchain](FRHICommandListBase&)
 			{
 				RenderBridge->RunOnRHISubmissionThread([MotionVectorSwapchain, MotionVectorDepthSwapchain]
