Fix crash in OpenXR swapchain image enumeration (UE-358632)

Add error handling to EnumerateImages() which was missing validation
on the first xrEnumerateSwapchainImages call (count query). The return
value was discarded and ChainCount left uninitialized on failure,
leading to a null pointer dereference inside third-party OpenXR
runtimes (e.g. picostreaming-openxr).

Also add empty-images guards in all five CreateSwapchain_* functions
(D3D11, D3D12, OpenGL, OpenGL ES, Vulkan) to prevent out-of-bounds
access on TextureChain[0] when enumeration returns zero images, and
properly destroy the leaked swapchain handle in that case.

#Synced-CL 52549286

diff --git a/Engine/Plugins/Runtime/OpenXR/Source/OpenXRHMD/Private/OpenXRHMD_Swapchain.cpp b/Engine/Plugins/Runtime/OpenXR/Source/OpenXRHMD/Private/OpenXRHMD_Swapchain.cpp
--- a/Engine/Plugins/Runtime/OpenXR/Source/OpenXRHMD/Private/OpenXRHMD_Swapchain.cpp
+++ b/Engine/Plugins/Runtime/OpenXR/Source/OpenXRHMD/Private/OpenXRHMD_Swapchain.cpp
@@ -255,15 +255,21 @@ template<typename T>
 TArray<T> EnumerateImages(XrSwapchain InSwapchain, XrStructureType InType, void* Next = nullptr)
 {
 	TArray<T> Images;
-	uint32_t ChainCount;
-	xrEnumerateSwapchainImages(InSwapchain, 0, &ChainCount, nullptr);
+	uint32_t ChainCount = 0;
+	if (!XR_ENSURE(xrEnumerateSwapchainImages(InSwapchain, 0, &ChainCount, nullptr)) || ChainCount == 0)
+	{
+		return Images;
+	}
 	Images.AddZeroed(ChainCount);
 	for (auto& Image : Images)
 	{
 		Image.type = InType;
 		Image.next = Next;
 	}
-	XR_ENSURE(xrEnumerateSwapchainImages(InSwapchain, ChainCount, &ChainCount, reinterpret_cast<XrSwapchainImageBaseHeader*>(Images.GetData())));
+	if (!XR_ENSURE(xrEnumerateSwapchainImages(InSwapchain, ChainCount, &ChainCount, reinterpret_cast<XrSwapchainImageBaseHeader*>(Images.GetData()))))
+	{
+		Images.Empty();
+	}
 	return Images;
 }
 
@@ -325,6 +331,12 @@ FXRSwapChainPtr CreateSwapchain_D3D11(XrSession InSession, uint8 Format, uint8&
 
 	TArray<FTextureRHIRef> TextureChain;
 	TArray<XrSwapchainImageD3D11KHR> Images = EnumerateImages<XrSwapchainImageD3D11KHR>(Swapchain, XR_TYPE_SWAPCHAIN_IMAGE_D3D11_KHR);
+	if (Images.IsEmpty())
+	{
+		UE_LOG(LogHMD, Error, TEXT("Failed to enumerate D3D11 swapchain images"));
+		xrDestroySwapchain(Swapchain);
+		return nullptr;
+	}
 	ETextureCreateFlags SanitizedCreateFlags = SanitizeTextureCreateFlags(CreateFlags);
 	for (const auto& Image : Images)
 	{
@@ -372,6 +384,12 @@ FXRSwapChainPtr CreateSwapchain_D3D12(XrSession InSession, uint8 Format, uint8&
 	ID3D12DynamicRHI* DynamicRHI = GetID3D12DynamicRHI();
 	TArray<FTextureRHIRef> TextureChain;
 	TArray<XrSwapchainImageD3D12KHR> Images = EnumerateImages<XrSwapchainImageD3D12KHR>(Swapchain, XR_TYPE_SWAPCHAIN_IMAGE_D3D12_KHR);
+	if (Images.IsEmpty())
+	{
+		UE_LOG(LogHMD, Error, TEXT("Failed to enumerate D3D12 swapchain images"));
+		xrDestroySwapchain(Swapchain);
+		return nullptr;
+	}
 	ETextureCreateFlags SanitizedCreateFlags = SanitizeTextureCreateFlags(CreateFlags);
 	for (const auto& Image : Images)
 	{
@@ -415,6 +433,12 @@ FXRSwapChainPtr CreateSwapchain_OpenGL(XrSession InSession, uint8 Format, uint8&
 	TArray<FTextureRHIRef> TextureChain;
 	IOpenGLDynamicRHI* DynamicRHI = GetIOpenGLDynamicRHI();
 	TArray<XrSwapchainImageOpenGLKHR> Images = EnumerateImages<XrSwapchainImageOpenGLKHR>(Swapchain, XR_TYPE_SWAPCHAIN_IMAGE_OPENGL_KHR);
+	if (Images.IsEmpty())
+	{
+		UE_LOG(LogHMD, Error, TEXT("Failed to enumerate OpenGL swapchain images"));
+		xrDestroySwapchain(Swapchain);
+		return nullptr;
+	}
 	ETextureCreateFlags SanitizedCreateFlags = SanitizeTextureCreateFlags(CreateFlags);
 	for (const auto& Image : Images)
 	{
@@ -467,6 +491,12 @@ FXRSwapChainPtr CreateSwapchain_OpenGLES(XrSession InSession, uint8 Format, uint
 	TArray<FTextureRHIRef> TextureChain;
 	IOpenGLDynamicRHI* DynamicRHI = GetIOpenGLDynamicRHI();
 	TArray<XrSwapchainImageOpenGLESKHR> Images = EnumerateImages<XrSwapchainImageOpenGLESKHR>(Swapchain, XR_TYPE_SWAPCHAIN_IMAGE_OPENGL_ES_KHR);
+	if (Images.IsEmpty())
+	{
+		UE_LOG(LogHMD, Error, TEXT("Failed to enumerate OpenGL ES swapchain images"));
+		xrDestroySwapchain(Swapchain);
+		return nullptr;
+	}
 	ETextureCreateFlags SanitizedCreateFlags = SanitizeTextureCreateFlags(CreateFlags);
 	for (const auto& Image : Images)
 	{
@@ -548,6 +578,12 @@ FXRSwapChainPtr CreateSwapchain_Vulkan(XrSession InSession, uint8 Format, uint8&
 
 	TArray<FTextureRHIRef> TextureChain;
 	TArray<XrSwapchainImageVulkanKHR> Images = EnumerateImages<XrSwapchainImageVulkanKHR>(Swapchain, XR_TYPE_SWAPCHAIN_IMAGE_VULKAN_KHR);
+	if (Images.IsEmpty())
+	{
+		UE_LOG(LogHMD, Error, TEXT("Failed to enumerate Vulkan swapchain images"));
+		xrDestroySwapchain(Swapchain);
+		return nullptr;
+	}
 	ETextureCreateFlags SanitizedCreateFlags = SanitizeTextureCreateFlags(CreateFlags);
 	for (const auto& Image : Images)
 	{
