Remove early out of FOpenXRHMD::AllocateRenderTargetTextures when !Swapchain.IsValid(), as this could cause a single frame to be rendered at launch without a valid XR swapchain, causing crashes or hangs.

#Synced-CL 52685484

diff --git a/Engine/Plugins/Runtime/OpenXR/Source/OpenXRHMD/Private/OpenXRHMD.cpp b/Engine/Plugins/Runtime/OpenXR/Source/OpenXRHMD/Private/OpenXRHMD.cpp
--- a/Engine/Plugins/Runtime/OpenXR/Source/OpenXRHMD/Private/OpenXRHMD.cpp
+++ b/Engine/Plugins/Runtime/OpenXR/Source/OpenXRHMD/Private/OpenXRHMD.cpp
@@ -2848,10 +2848,10 @@ bool FOpenXRHMD::AllocateRenderTargetTextures(FRHICommandListBase& RHICmdList, u
 	}
 
 	FXRSwapChainPtr& Swapchain = PipelinedLayerStateRendering.ColorSwapchain;
-	if (!Swapchain.IsValid())
+	/*if (!Swapchain.IsValid())
 	{
 		return false;
-	}
+	}*/
 
 	uint8 ActualFormat = Format;
 	ensureMsgf(NumSamples == 1, TEXT("OpenXR supports MSAA swapchains, but engine logic expects the swapchain target to be 1x."));
